Added pcg external data support plugin, implemented a native data table loader and added alembic file support

diff --git a/Engine/Plugins/Experimental/PCG/Source/PCG/Private/Elements/IO/PCGDataTableElement.cpp b/Engine/Plugins/Experimental/PCG/Source/PCG/Private/Elements/IO/PCGDataTableElement.cpp
new file mode 100644
--- /dev/null
+++ b/Engine/Plugins/Experimental/PCG/Source/PCG/Private/Elements/IO/PCGDataTableElement.cpp
@@ -0,0 +1,183 @@
+// Copyright Epic Games, Inc. All Rights Reserved.
+
+#include "Elements/IO/PCGDataTableElement.h"
+
+#include "PCGComponent.h"
+#include "PCGModule.h"
+#include "Data/PCGPointData.h"
+#include "Elements/IO/PCGExternalDataContext.h"
+
+#define LOCTEXT_NAMESPACE "PCGLoadDataTable"
+
+#if WITH_EDITOR
+
+FText UPCGLoadDataTableSettings::GetDefaultNodeTitle() const
+{
+	return LOCTEXT("NodeTitle", "Load Data Table");
+}
+
+FText UPCGLoadDataTableSettings::GetNodeTooltipText() const
+{
+	return LOCTEXT("NodeTooltip", "Loads data from DataTable asset");
+}
+
+#endif // WITH_EDITOR
+
+FPCGElementPtr UPCGLoadDataTableSettings::CreateElement() const
+{
+	return MakeShared<FPCGLoadDataTableElement>();
+}
+
+bool FPCGLoadDataTableElement::PrepareLoad(FPCGExternalDataContext* Context) const
+{
+	check(Context);
+	const UPCGLoadDataTableSettings* Settings = Context->GetInputSettings<UPCGLoadDataTableSettings>();
+	check(Settings);
+
+	// 1. Load data table - if we do, we should root it (?)
+	UDataTable* DataTable = Settings->DataTable.LoadSynchronous();
+	if (!DataTable)
+	{
+		PCGE_LOG(Error, GraphAndLog, FText::Format(LOCTEXT("DataTableNotLoaded", "Provided Data Table ('{0}') could not be loaded."), FText::FromString(Settings->DataTable.ToString())));
+		return true;
+	}
+
+	// 2. Make sure it has a row struct
+	const UScriptStruct* RowStruct = DataTable->GetRowStruct();
+	if (!RowStruct)
+	{
+		PCGE_LOG(Error, GraphAndLog, FText::Format(LOCTEXT("DataTableNoRowStruct", "Provided Data Table ('{0}') does not have a valid row definition."), FText::FromString(Settings->DataTable.ToString())));
+		return true;
+	}
+
+	// 2.1 Validate that there are rows too
+	if (DataTable->GetRowMap().IsEmpty())
+	{
+		PCGE_LOG(Warning, GraphAndLog, FText::Format(LOCTEXT("DataTableNoRowStruct", "Provided Data Table ('{0}') is empty."), FText::FromString(Settings->DataTable.ToString())));
+		return true;
+	}
+
+	// 3-e. Create point data, and the mapping object
+	UPCGPointData* PointData = NewObject<UPCGPointData>();
+	check(PointData);
+
+	if (Context->SourceComponent.Get())
+	{
+		PointData->TargetActor = Context->SourceComponent->GetOwner();
+	}
+
+	FPCGExternalDataContext::FPointDataAccessorsMapping& PointDataAccessorMapping = Context->PointDataAccessorsMapping.Emplace_GetRef();
+	PointDataAccessorMapping.PointData = PointData;
+
+	UPCGMetadata* PointMetadata = PointData->MutableMetadata();
+
+	// 3. Build property to accessor mappings, including renamings, and create new attributes as needed
+	for (TFieldIterator<const FProperty> FieldIt(RowStruct, EFieldIterationFlags::IncludeSuper); FieldIt; ++FieldIt)
+	{
+		FString FieldName = RowStruct->GetAuthoredNameForField(*FieldIt);
+
+		// Try to create the read accessor
+		TUniquePtr<IPCGAttributeAccessor> DataAccessor = PCGAttributeAccessorHelpers::CreatePropertyAccessor(*FieldIt);
+
+		if (!DataAccessor)
+		{
+			PCGE_LOG(Error, GraphAndLog, FText::Format(LOCTEXT("DataTableFieldTypeNotSupported", "Provided Data Table ('{0}') property '{1}' is not of a supported type."), FText::FromString(Settings->DataTable.ToString()), FText::FromString(RowStruct->GetAuthoredNameForField(*FieldIt))));
+			continue;
+		}
+
+		// Try to match to a property (if the name maps to a property) or create attribute as needed
+		FPCGAttributePropertySelector PointPropertySelector;
+
+		if (const FPCGAttributePropertySelector* MappedField = Settings->AttributeMapping.Find(FieldName))
+		{
+			PointPropertySelector = *MappedField;
+			FieldName = PointPropertySelector.GetDisplayText().ToString();
+		}
+		else
+		{
+			PointPropertySelector.Update(FieldName);
+		}
+
+		if (PointPropertySelector.Selection == EPCGAttributePropertySelection::Attribute)
+		{
+			// Create attribute with the appropriate type
+			PointMetadata->CreateAttributeFromProperty(FName(FieldName), RowStruct, *FieldIt);
+		}
+
+		TUniquePtr<IPCGAttributeAccessor> PointPropertyAccessor = PCGAttributeAccessorHelpers::CreateAccessor(PointData, PointPropertySelector);
+
+		if (!PointPropertyAccessor)
+		{
+			PCGE_LOG(Error, GraphAndLog, FText::Format(LOCTEXT("UnableToCreatePointProperty", "Unable to create/access point property/attribute '{0}'."), FText::FromString(FieldName)));
+			continue;
+		}
+
+		// Final verification, if we can put the value of input into output
+		if (!PCG::Private::IsBroadcastable(DataAccessor->GetUnderlyingType(), PointPropertyAccessor->GetUnderlyingType()))
+		{
+			FText InputTypeName = FText::FromString(PCG::Private::GetTypeName(DataAccessor->GetUnderlyingType()));
+			FText OutputTypeName = FText::FromString(PCG::Private::GetTypeName(PointPropertyAccessor->GetUnderlyingType()));
+
+			PCGE_LOG_C(Error, GraphAndLog, Context, FText::Format(LOCTEXT("CannotBroadcastTypes", "Cannot convert input type '{0}' into output type '{1}'"), InputTypeName, OutputTypeName));
+			continue;
+		}
+
+		PointDataAccessorMapping.RowToPointAccessors.Emplace(MoveTemp(DataAccessor), MoveTemp(PointPropertyAccessor), PointPropertySelector);
+	}
+
+	if (!Context->PointDataAccessorsMapping.IsEmpty())
+	{
+		FPCGTaggedData& OutData = Context->OutputData.TaggedData.Emplace_GetRef();
+		OutData.Data = PointData;
+
+		Context->bDataPrepared = true;
+	}
+
+	return true;
+}
+
+
+bool FPCGLoadDataTableElement::ExecuteLoad(FPCGExternalDataContext* Context) const
+{
+	check(Context);
+
+	if (Context->PointDataAccessorsMapping.IsEmpty())
+	{
+		return true;
+	}
+
+	const UPCGLoadDataTableSettings* Settings = Context->GetInputSettings<UPCGLoadDataTableSettings>();
+	check(Settings);
+
+	//TODO: instead of trying to reload here, it would be better to make sure that this is rooted properly
+	UDataTable* DataTable = Settings->DataTable.LoadSynchronous();
+	check(DataTable);
+
+	// Note that while this is prepared for execution in multiple calls, we need to make sure the data table is kept in memory at all times otherwise the row keys could become stale
+	check(Context->PointDataAccessorsMapping.Num() == 1);
+	FPCGExternalDataContext::FPointDataAccessorsMapping& PointDataAccessorMapping = Context->PointDataAccessorsMapping[0];
+
+	// Preallocate the point array, and create the row keys
+	UPCGPointData* PointData = PointDataAccessorMapping.PointData;
+	check(PointData);
+
+	TArray<FPCGPoint>& Points = PointDataAccessorMapping.PointData->GetMutablePoints();
+	if (Points.IsEmpty())
+	{
+		// Dummy implementation - note that we do initialize the points here in cases where some of the properties would not be set
+		Points.SetNum(DataTable->GetRowMap().Num());
+	}
+
+	TArray<void*> AllRows;
+	AllRows.Reserve(DataTable->GetRowMap().Num());
+	for (TMap<FName, uint8*>::TConstIterator RowMapIter(DataTable->GetRowMap().CreateConstIterator()); RowMapIter; ++RowMapIter)
+	{
+		AllRows.Add(RowMapIter.Value());
+	}
+
+	PointDataAccessorMapping.RowKeys = MakeUnique<FPCGAttributeAccessorKeysGenericPtrs>(AllRows);
+	
+	return FPCGExternalDataElement::ExecuteLoad(Context);
+}
+
+#undef LOCTEXT_NAMESPACE
\ No newline at end of file
diff --git a/Engine/Plugins/Experimental/PCG/Source/PCG/Private/Elements/IO/PCGExternalData.cpp b/Engine/Plugins/Experimental/PCG/Source/PCG/Private/Elements/IO/PCGExternalData.cpp
new file mode 100644
--- /dev/null
+++ b/Engine/Plugins/Experimental/PCG/Source/PCG/Private/Elements/IO/PCGExternalData.cpp
@@ -0,0 +1,97 @@
+// Copyright Epic Games, Inc. All Rights Reserved.
+
+#include "Elements/IO/PCGExternalData.h"
+#include "Elements/IO/PCGExternalDataContext.h"
+
+#include "PCGComponent.h"
+#include "Data/PCGPointData.h"
+
+#define LOCTEXT_NAMESPACE "PCGExternalData"
+
+FPCGContext* FPCGExternalDataElement::Initialize(const FPCGDataCollection& InputData, TWeakObjectPtr<UPCGComponent> SourceComponent, const UPCGNode* Node)
+{
+	FPCGExternalDataContext* Context = new FPCGExternalDataContext();
+	Context->InputData = InputData;
+	Context->SourceComponent = SourceComponent;
+	Context->Node = Node;
+
+	return Context;
+}
+
+bool FPCGExternalDataElement::PrepareDataInternal(FPCGContext* InContext) const
+{
+	TRACE_CPUPROFILER_EVENT_SCOPE(FPCGExternalDataElement::PrepareData);
+	FPCGExternalDataContext* Context = static_cast<FPCGExternalDataContext*>(InContext);
+	check(Context);
+
+	return PrepareLoad(Context);
+}
+
+bool FPCGExternalDataElement::ExecuteInternal(FPCGContext* InContext) const
+{
+	TRACE_CPUPROFILER_EVENT_SCOPE(FPCGExternalDataElement::Execute);
+	FPCGExternalDataContext* Context = static_cast<FPCGExternalDataContext*>(InContext);
+	check(Context);
+
+	if (!Context->bDataPrepared)
+	{
+		return true;
+	}
+
+	return ExecuteLoad(Context);
+}
+
+bool FPCGExternalDataElement::ExecuteLoad(FPCGExternalDataContext* Context) const
+{
+	check(Context);
+
+	// Do in parallel according to the resources given in the context
+	for (FPCGExternalDataContext::FPointDataAccessorsMapping& PointDataAccessorMapping : Context->PointDataAccessorsMapping)
+	{
+		UPCGPointData* PointData = PointDataAccessorMapping.PointData;
+		UPCGMetadata* PointMetadata = PointData->MutableMetadata();
+		TUniquePtr<const IPCGAttributeAccessorKeys>& RowKeys = PointDataAccessorMapping.RowKeys;
+
+		for (FPCGExternalDataContext::FRowToPointAccessors& RowToPointAccessor : PointDataAccessorMapping.RowToPointAccessors)
+		{
+			TUniquePtr<IPCGAttributeAccessorKeys> PointKeys = PCGAttributeAccessorHelpers::CreateKeys(PointData, RowToPointAccessor.Selector);
+
+			auto Operation = [&RowKeys, &RowToPointAccessor, &PointKeys](auto Dummy)
+			{
+				using OutputType = decltype(Dummy);
+				OutputType Value{};
+
+				EPCGAttributeAccessorFlags Flags = EPCGAttributeAccessorFlags::AllowBroadcast;
+
+				const int32 NumberOfElements = PointKeys->GetNum();
+				constexpr int32 ChunkSize = 256;
+
+				TArray<OutputType, TInlineAllocator<ChunkSize>> TempValues;
+				TempValues.SetNum(ChunkSize);
+
+				const int32 NumberOfIterations = (NumberOfElements + ChunkSize - 1) / ChunkSize;
+
+				for (int32 i = 0; i < NumberOfIterations; ++i)
+				{
+					const int32 StartIndex = i * ChunkSize;
+					const int32 Range = FMath::Min(NumberOfElements - StartIndex, ChunkSize);
+					TArrayView<OutputType> View(TempValues.GetData(), Range);
+
+					RowToPointAccessor.RowAccessor->GetRange<OutputType>(View, StartIndex, *RowKeys, Flags);
+					RowToPointAccessor.PointAccessor->SetRange<OutputType>(View, StartIndex, *PointKeys, Flags);
+				}
+
+				return true;
+			};
+
+			if (!PCGMetadataAttribute::CallbackWithRightType(RowToPointAccessor.PointAccessor->GetUnderlyingType(), Operation))
+			{
+				PCGE_LOG(Warning, GraphAndLog, LOCTEXT("ErrorGettingSettingValues", "Error while getting/setting values"));
+			}
+		}
+	}
+
+	return true;
+}
+
+#undef LOCTEXT_NAMESPACE
\ No newline at end of file
diff --git a/Engine/Plugins/Experimental/PCG/Source/PCG/Private/Elements/Metadata/PCGMetadataRotatorOpElement.cpp b/Engine/Plugins/Experimental/PCG/Source/PCG/Private/Elements/Metadata/PCGMetadataRotatorOpElement.cpp
--- a/Engine/Plugins/Experimental/PCG/Source/PCG/Private/Elements/Metadata/PCGMetadataRotatorOpElement.cpp
+++ b/Engine/Plugins/Experimental/PCG/Source/PCG/Private/Elements/Metadata/PCGMetadataRotatorOpElement.cpp
@@ -29,7 +29,8 @@ namespace PCGMetadataRotatorSettings
 {
 	inline constexpr bool IsUnaryOp(EPCGMedadataRotatorOperation Operation)
 	{
-		return Operation == EPCGMedadataRotatorOperation::Invert;
+		return Operation == EPCGMedadataRotatorOperation::Invert ||
+			Operation == EPCGMedadataRotatorOperation::Normalize;
 	}
 
 	inline constexpr bool IsTernaryOp(EPCGMedadataRotatorOperation Operation)
@@ -52,6 +53,8 @@ namespace PCGMetadataRotatorSettings
 			return PCGMetadataRotatorHelpers::RLerp(Input1, Input2, Ratio, false);
 		case EPCGMedadataRotatorOperation::Invert:
 			return Input1.GetInverse();
+		case EPCGMedadataRotatorOperation::Normalize:
+			return Input1.GetNormalized();
 		default:
 			return FRotator{};
 		}
@@ -67,6 +70,8 @@ namespace PCGMetadataRotatorSettings
 			return FQuat::Slerp(Input1, Input2, Ratio);
 		case EPCGMedadataRotatorOperation::Invert:
 			return Input1.Inverse();
+		case EPCGMedadataRotatorOperation::Normalize:
+			return Input1.GetNormalized();
 		default:
 			return FQuat{};
 		}
diff --git a/Engine/Plugins/Experimental/PCG/Source/PCG/Private/Metadata/Accessors/PCGAttributeExtractor.cpp b/Engine/Plugins/Experimental/PCG/Source/PCG/Private/Metadata/Accessors/PCGAttributeExtractor.cpp
--- a/Engine/Plugins/Experimental/PCG/Source/PCG/Private/Metadata/Accessors/PCGAttributeExtractor.cpp
+++ b/Engine/Plugins/Experimental/PCG/Source/PCG/Private/Metadata/Accessors/PCGAttributeExtractor.cpp
@@ -47,12 +47,16 @@ namespace PCGAttributeExtractor
 		{
 		case 0:
 			Value.X = In;
+			break;
 		case 1:
 			Value.Y = In;
+			break;
 		case 2:
 			Value.Z = In;
+			break;
 		default:
 			Value.W = In;
+			break;
 		}
 	}
 
@@ -66,6 +70,21 @@ namespace PCGAttributeExtractor
 			return MakeUnique<FPCGChainAccessor<double, VectorType>>(std::move(InAccessor),
 				[](const VectorType& Value) -> double { return Value.Size(); });
 		}
+		else if (Name == PCGAttributeExtractorConstants::VectorNormalized)
+		{
+			if constexpr (std::is_same_v<FQuat, VectorType>)
+			{
+				bOutSuccess = true;
+				return MakeUnique<FPCGChainAccessor<VectorType, VectorType>>(std::move(InAccessor),
+					[](const VectorType& Value) -> VectorType { return Value.GetNormalized(); });
+			}
+			else
+			{
+				bOutSuccess = true;
+				return MakeUnique<FPCGChainAccessor<VectorType, VectorType>>(std::move(InAccessor),
+					[](const VectorType& Value) -> VectorType { return Value.GetSafeNormal(); });
+			}
+		}
 
 		const FString NameStr = Name.ToString();
 		const int32 NameLen = NameStr.Len();
diff --git a/Engine/Plugins/Experimental/PCG/Source/PCG/Private/Metadata/PCGAttributePropertySelector.cpp b/Engine/Plugins/Experimental/PCG/Source/PCG/Private/Metadata/PCGAttributePropertySelector.cpp
--- a/Engine/Plugins/Experimental/PCG/Source/PCG/Private/Metadata/PCGAttributePropertySelector.cpp
+++ b/Engine/Plugins/Experimental/PCG/Source/PCG/Private/Metadata/PCGAttributePropertySelector.cpp
@@ -132,7 +132,6 @@ FText FPCGAttributePropertySelector::GetDisplayText() const
 	return FText::FromString(Res);
 }
 
-#if WITH_EDITOR
 bool FPCGAttributePropertySelector::IsValid() const
 {
 	return (Selection != EPCGAttributePropertySelection::Attribute) || FPCGMetadataAttributeBase::IsValidName(AttributeName);
@@ -186,7 +185,6 @@ bool FPCGAttributePropertySelector::Update(FString NewValue)
 
 	return SetAttributeName(NewName.IsEmpty() ? NAME_None : FName(NewName), /*bResetExtraNames=*/ false) || bExtraChanged;
 }
-#endif // WITH_EDITOR
 
 ///////////////////////////////////////////////////////////////////////
 
diff --git a/Engine/Plugins/Experimental/PCG/Source/PCG/Private/Metadata/PCGMetadata.cpp b/Engine/Plugins/Experimental/PCG/Source/PCG/Private/Metadata/PCGMetadata.cpp
--- a/Engine/Plugins/Experimental/PCG/Source/PCG/Private/Metadata/PCGMetadata.cpp
+++ b/Engine/Plugins/Experimental/PCG/Source/PCG/Private/Metadata/PCGMetadata.cpp
@@ -515,6 +515,25 @@ void UPCGMetadata::CreateBoolAttribute(FName AttributeName, bool DefaultValue, b
 
 namespace PCGMetadata
 {
+	template<typename DataType>
+	bool CreateAttributeFromPropertyHelper(UPCGMetadata* Metadata, FName AttributeName, const DataType* DataPtr, const FProperty* InProperty)
+	{
+		if (!InProperty || !DataPtr || !Metadata)
+		{
+			return false;
+		}
+
+		auto CreateAttribute = [AttributeName, Metadata](auto&& PropertyValue) -> bool
+		{
+			using PropertyType = std::decay_t<decltype(PropertyValue)>;
+			FPCGMetadataAttributeBase* BaseAttribute = Metadata->FindOrCreateAttribute<PropertyType>(AttributeName, PropertyValue, /*bAllowsInterpolation=*/false, /*bOverrideParent=*/false, /*bOverwriteIfTypeMismatch=*/true);
+			
+			return (BaseAttribute != nullptr);
+		};
+
+		return PCGPropertyHelpers::GetPropertyValueWithCallback(DataPtr, InProperty, CreateAttribute);
+	}
+
 	template<typename DataType>
 	bool SetAttributeFromPropertyHelper(UPCGMetadata* Metadata, FName AttributeName, PCGMetadataEntryKey& EntryKey, const DataType* DataPtr, const FProperty* InProperty, bool bCreate)
 	{
@@ -575,6 +594,16 @@ namespace PCGMetadata
 	}
 }
 
+bool UPCGMetadata::CreateAttributeFromProperty(FName AttributeName, const UObject* Object, const FProperty* InProperty)
+{
+	return PCGMetadata::CreateAttributeFromPropertyHelper<UObject>(this, AttributeName, Object, InProperty);
+}
+
+bool UPCGMetadata::CreateAttributeFromDataProperty(FName AttributeName, const void* Data, const FProperty* InProperty)
+{
+	return PCGMetadata::CreateAttributeFromPropertyHelper<void>(this, AttributeName, Data, InProperty);
+}
+
 bool UPCGMetadata::SetAttributeFromProperty(FName AttributeName, PCGMetadataEntryKey& EntryKey, const UObject* Object, const FProperty* InProperty, bool bCreate)
 {
 	return PCGMetadata::SetAttributeFromPropertyHelper<UObject>(this, AttributeName, EntryKey, Object, InProperty, bCreate);
diff --git a/Engine/Plugins/Experimental/PCG/Source/PCG/Public/Elements/IO/PCGDataTableElement.h b/Engine/Plugins/Experimental/PCG/Source/PCG/Public/Elements/IO/PCGDataTableElement.h
new file mode 100644
--- /dev/null
+++ b/Engine/Plugins/Experimental/PCG/Source/PCG/Public/Elements/IO/PCGDataTableElement.h
@@ -0,0 +1,36 @@
+// Copyright Epic Games, Inc. All Rights Reserved.
+
+#pragma once
+
+#include "PCGExternalData.h"
+
+#include "PCGDataTableElement.generated.h"
+
+UCLASS(BlueprintType, ClassGroup = (Procedural))
+class PCG_API UPCGLoadDataTableSettings : public UPCGExternalDataSettings
+{
+	GENERATED_BODY()
+
+public:
+	// ~Begin UPCGSettings interface
+#if WITH_EDITOR
+	virtual FName GetDefaultNodeName() const override { return FName(TEXT("LoadDataTable")); }
+	virtual FText GetDefaultNodeTitle() const override;
+	virtual FText GetNodeTooltipText() const override;
+#endif
+
+protected:
+	virtual FPCGElementPtr CreateElement() const override;
+	// ~End UPCGSettings interface
+
+public:
+	UPROPERTY(BlueprintReadWrite, EditAnywhere, Category = "Settings", meta = (PCG_Overridable))
+	TSoftObjectPtr<UDataTable> DataTable;
+};
+
+class FPCGLoadDataTableElement : public FPCGExternalDataElement
+{
+protected:
+	virtual bool PrepareLoad(FPCGExternalDataContext* Context) const override;
+	virtual bool ExecuteLoad(FPCGExternalDataContext* Context) const override;
+};
\ No newline at end of file
diff --git a/Engine/Plugins/Experimental/PCG/Source/PCG/Public/Elements/IO/PCGExternalData.h b/Engine/Plugins/Experimental/PCG/Source/PCG/Public/Elements/IO/PCGExternalData.h
new file mode 100644
--- /dev/null
+++ b/Engine/Plugins/Experimental/PCG/Source/PCG/Public/Elements/IO/PCGExternalData.h
@@ -0,0 +1,48 @@
+// Copyright Epic Games, Inc. All Rights Reserved.
+
+#pragma once
+
+#include "PCGContext.h"
+#include "PCGElement.h"
+#include "PCGPin.h"
+#include "PCGSettings.h"
+#include "Metadata/PCGAttributePropertySelector.h"
+
+#include "PCGExternalData.generated.h"
+
+struct FPCGExternalDataContext;
+
+/** Base class for external data input settings */
+UCLASS(Abstract, ClassGroup = (Procedural))
+class PCG_API UPCGExternalDataSettings : public UPCGSettings
+{
+	GENERATED_BODY()
+
+public:
+	// ~Begin UPCGSettings interface
+#if WITH_EDITOR
+	virtual EPCGSettingsType GetType() const override { return EPCGSettingsType::InputOutput; }
+#endif
+
+protected:
+	virtual TArray<FPCGPinProperties> InputPinProperties() const override { return TArray<FPCGPinProperties>(); }
+	virtual TArray<FPCGPinProperties> OutputPinProperties() const override { return DefaultPointOutputPinProperties(); }
+	// ~End UPCGSettings interface
+
+public:
+	UPROPERTY(BlueprintReadWrite, EditAnywhere, Category = "Settings")
+	TMap<FString, FPCGAttributePropertySelector> AttributeMapping;
+};
+
+class PCG_API FPCGExternalDataElement : public IPCGElement
+{
+public:
+	virtual FPCGContext* Initialize(const FPCGDataCollection& InputData, TWeakObjectPtr<UPCGComponent> SourceComponent, const UPCGNode* Node) override;
+protected:
+	virtual bool PrepareDataInternal(FPCGContext* InContext) const override;
+	virtual bool ExecuteInternal(FPCGContext* InContext) const override;
+
+	virtual bool PrepareLoad(FPCGExternalDataContext* Context) const = 0;
+	virtual bool ExecuteLoad(FPCGExternalDataContext* Context) const;
+	virtual bool PostExecuteLoad(FPCGExternalDataContext* Context) const { return true; }
+};
\ No newline at end of file
diff --git a/Engine/Plugins/Experimental/PCG/Source/PCG/Public/Elements/IO/PCGExternalDataContext.h b/Engine/Plugins/Experimental/PCG/Source/PCG/Public/Elements/IO/PCGExternalDataContext.h
new file mode 100644
--- /dev/null
+++ b/Engine/Plugins/Experimental/PCG/Source/PCG/Public/Elements/IO/PCGExternalDataContext.h
@@ -0,0 +1,40 @@
+// Copyright Epic Games, Inc. All Rights Reserved.
+
+#pragma once
+
+#include "PCGContext.h"
+#include "Data/PCGPointData.h"
+#include "Metadata/Accessors/IPCGAttributeAccessor.h"
+#include "Metadata/Accessors/PCGAttributeAccessorHelpers.h"
+#include "Metadata/Accessors/PCGAttributeAccessorKeys.h"
+
+struct FPCGExternalDataContext : public FPCGContext
+{
+	// Used for data table parsing
+	struct FRowToPointAccessors
+	{
+		TUniquePtr<const IPCGAttributeAccessor> RowAccessor;
+		TUniquePtr<IPCGAttributeAccessor> PointAccessor;
+		FPCGAttributePropertySelector Selector;
+
+		FRowToPointAccessors() = default;
+		FRowToPointAccessors(const FRowToPointAccessors&) = delete;
+		FRowToPointAccessors(FRowToPointAccessors&&) = default;
+		FRowToPointAccessors& operator=(FRowToPointAccessors&&) = default;
+		FRowToPointAccessors& operator=(const FRowToPointAccessors&) = delete;
+		FRowToPointAccessors(TUniquePtr<const IPCGAttributeAccessor> InRow, TUniquePtr<IPCGAttributeAccessor> InPoint, const FPCGAttributePropertySelector& InSelector)
+			: RowAccessor(MoveTemp(InRow)), PointAccessor(MoveTemp(InPoint)), Selector(InSelector) {}
+	};
+
+	struct FPointDataAccessorsMapping
+	{
+		UPCGPointData* PointData = nullptr;
+		TArray<FRowToPointAccessors> RowToPointAccessors;
+		TUniquePtr<const IPCGAttributeAccessorKeys> RowKeys;
+	};
+
+	TArray<FPointDataAccessorsMapping> PointDataAccessorsMapping;
+
+	// Set to true if the prepare data succeeded
+	bool bDataPrepared = false;
+};
\ No newline at end of file
diff --git a/Engine/Plugins/Experimental/PCG/Source/PCG/Public/Elements/Metadata/PCGMetadataRotatorOpElement.h b/Engine/Plugins/Experimental/PCG/Source/PCG/Public/Elements/Metadata/PCGMetadataRotatorOpElement.h
--- a/Engine/Plugins/Experimental/PCG/Source/PCG/Public/Elements/Metadata/PCGMetadataRotatorOpElement.h
+++ b/Engine/Plugins/Experimental/PCG/Source/PCG/Public/Elements/Metadata/PCGMetadataRotatorOpElement.h
@@ -13,6 +13,7 @@ enum class EPCGMedadataRotatorOperation : uint16
 	Combine,
 	Invert,
 	Lerp,
+	Normalize,
 
 	TransformOp = 100 UMETA(Hidden),
 	TransformRotation,
diff --git a/Engine/Plugins/Experimental/PCG/Source/PCG/Public/Metadata/Accessors/PCGAttributeExtractor.h b/Engine/Plugins/Experimental/PCG/Source/PCG/Public/Metadata/Accessors/PCGAttributeExtractor.h
--- a/Engine/Plugins/Experimental/PCG/Source/PCG/Public/Metadata/Accessors/PCGAttributeExtractor.h
+++ b/Engine/Plugins/Experimental/PCG/Source/PCG/Public/Metadata/Accessors/PCGAttributeExtractor.h
@@ -20,6 +20,7 @@ namespace PCGAttributeExtractorConstants
 	const FName VectorA = TEXT("A");
 	const FName VectorLength = TEXT("Length");
 	const FName VectorSize = TEXT("Size");
+	const FName VectorNormalized = TEXT("Normalized");
 
 	const FName VectorExtractorNames[] =
 	{
diff --git a/Engine/Plugins/Experimental/PCG/Source/PCG/Public/Metadata/PCGAttributePropertySelector.h b/Engine/Plugins/Experimental/PCG/Source/PCG/Public/Metadata/PCGAttributePropertySelector.h
--- a/Engine/Plugins/Experimental/PCG/Source/PCG/Public/Metadata/PCGAttributePropertySelector.h
+++ b/Engine/Plugins/Experimental/PCG/Source/PCG/Public/Metadata/PCGAttributePropertySelector.h
@@ -62,13 +62,11 @@ public:
 	// Returns the text to display in the widget.
 	FText GetDisplayText() const;
 
-#if WITH_EDITOR
 	// Return true if the underlying name is valid.
 	bool IsValid() const;
 
 	// Update the selector with an incoming string.
 	bool Update(FString NewValue);
-#endif // WITH_EDITOR
 };
 
 /**
diff --git a/Engine/Plugins/Experimental/PCG/Source/PCG/Public/Metadata/PCGMetadata.h b/Engine/Plugins/Experimental/PCG/Source/PCG/Public/Metadata/PCGMetadata.h
--- a/Engine/Plugins/Experimental/PCG/Source/PCG/Public/Metadata/PCGMetadata.h
+++ b/Engine/Plugins/Experimental/PCG/Source/PCG/Public/Metadata/PCGMetadata.h
@@ -124,6 +124,22 @@ public:
 	UFUNCTION(BlueprintCallable, Category = "PCG|Metadata")
 	void CreateBoolAttribute(FName AttributeName, bool DefaultValue, bool bAllowsInterpolation, bool bOverrideParent = true);
 
+	/** Creates an attribute given a property.
+	* @param AttributeName: Target attribute to create
+	* @param Object: Object to get the property value from
+	* @param Property: The property to set from
+	* @returns true if the attribute creation succeeded
+	*/
+	bool CreateAttributeFromProperty(FName AttributeName, const UObject* Object, const FProperty* Property);
+
+	/** Creates an attribute given a property.
+	* @param AttributeName: Target attribute to create
+	* @param Data: Data pointer to get the property value from
+	* @param Property: The property to set from
+	* @returns true if the attribute creation succeeded
+	*/
+	bool CreateAttributeFromDataProperty(FName AttributeName, const void* Data, const FProperty* Property);
+
 	/** Set an attribute given a property and its value.
 	* @param AttributeName: Target attribute to set the property's value to
 	* @param EntryKey: Metadata entry key to set the value to
diff --git a/Engine/Plugins/Experimental/PCGExternalDataInterop/PCGExternalDataInterop.uplugin b/Engine/Plugins/Experimental/PCGExternalDataInterop/PCGExternalDataInterop.uplugin
new file mode 100644
--- /dev/null
+++ b/Engine/Plugins/Experimental/PCGExternalDataInterop/PCGExternalDataInterop.uplugin
@@ -0,0 +1,39 @@
+{
+	"FileVersion": 3,
+	"Version": 1,
+	"VersionName": "0.1",
+	"FriendlyName": "Procedural Content Generation Framework (PCG) External Data Interop",
+	"Description": "Extra plugin for Procedural Content Generation Framework interacting with external data formats.",
+	"Category": "Editor",
+	"CreatedBy": "Epic Games, Inc.",
+	"CreatedByURL": "https://epicgames.com",
+	"DocsURL": "",
+	"MarketplaceURL": "",
+	"SupportURL": "",
+	"EnabledByDefault": false,
+	"CanContainContent": true,
+	"IsBetaVersion": false,
+	"Installed": false,
+	"Modules": [
+		{
+			"Name": "PCGExternalDataInterop",
+			"Type": "Runtime",
+			"LoadingPhase": "Default"
+		},
+		{
+			"Name": "PCGExternalDataInteropEditor",
+			"Type": "Editor",
+			"LoadingPhase": "Default"
+		},
+	],
+	"Plugins": [
+		{
+			"Name": "PCG",
+			"Enabled": true
+		},
+		{
+			"Name": "AlembicImporter",
+			"Enabled": true
+		}
+	]
+}
diff --git a/Engine/Plugins/Experimental/PCGExternalDataInterop/Source/PCGExternalDataInterop/PCGExternalDataInterop.Build.cs b/Engine/Plugins/Experimental/PCGExternalDataInterop/Source/PCGExternalDataInterop/PCGExternalDataInterop.Build.cs
new file mode 100644
--- /dev/null
+++ b/Engine/Plugins/Experimental/PCGExternalDataInterop/Source/PCGExternalDataInterop/PCGExternalDataInterop.Build.cs
@@ -0,0 +1,33 @@
+// Copyright Epic Games, Inc. All Rights Reserved.
+
+namespace UnrealBuildTool.Rules
+{
+	public class PCGExternalDataInterop : ModuleRules
+	{
+		public PCGExternalDataInterop(ReadOnlyTargetRules Target) : base(Target)
+		{
+			PCHUsage = PCHUsageMode.UseExplicitOrSharedPCHs;
+
+			PublicDependencyModuleNames.AddRange(
+				new string[] {
+					"Core",
+					"CoreUObject",
+					"Engine",
+					"Projects",
+					"PCG"
+				});
+
+			if (Target.bBuildEditor == true)
+			{
+				PrivateDependencyModuleNames.AddRange(
+					new string[]
+					{
+						"UnrealEd",
+						"AlembicLib",
+						"AlembicLibrary"
+					}
+				);
+			}
+		}
+	}
+}
diff --git a/Engine/Plugins/Experimental/PCGExternalDataInterop/Source/PCGExternalDataInterop/Private/Alembic/PCGAlembicInterop.cpp b/Engine/Plugins/Experimental/PCGExternalDataInterop/Source/PCGExternalDataInterop/Private/Alembic/PCGAlembicInterop.cpp
new file mode 100644
--- /dev/null
+++ b/Engine/Plugins/Experimental/PCGExternalDataInterop/Source/PCGExternalDataInterop/Private/Alembic/PCGAlembicInterop.cpp
@@ -0,0 +1,312 @@
+// Copyright Epic Games, Inc. All Rights Reserved.
+
+#include "PCGAlembicInterop.h"
+
+#include "Elements/PCGLoadAlembicElement.h"
+
+#include "PCGComponent.h"
+#include "Data/PCGPointData.h"
+#include "Elements/Metadata/PCGMetadataElementCommon.h"
+#include "Metadata/Accessors/IPCGAttributeAccessorTpl.h"
+#include "Metadata/Accessors/PCGAttributeAccessorHelpers.h"
+#include "Metadata/Accessors/PCGAttributeAccessorKeys.h"
+
+#include <string>
+
+#define LOCTEXT_NAMESPACE "PCGExternalData_Alembic"
+
+namespace PCGAlembicInterop
+{
+
+#if WITH_EDITOR
+template<typename T, typename AbcParamType, int32 Extent, int32 SubExtent, bool bUseCStr>
+class FPCGAlembicAccessor : public IPCGAttributeAccessorT<FPCGAlembicAccessor<T, AbcParamType, Extent, SubExtent, bUseCStr>>
+{
+public:
+	using Type = T;
+	using Super = IPCGAttributeAccessorT<FPCGAlembicAccessor<T, AbcParamType, Extent, SubExtent, bUseCStr>>;
+
+	FPCGAlembicAccessor(Alembic::AbcGeom::ICompoundProperty& Parameters, const FString& PropName)
+		: Super(/*bInReadOnly=*/true)
+		, Param(Parameters, std::string(TCHAR_TO_UTF8(*PropName)))
+		, SamplePtr(Param.getExpandedValue().getVals())
+	{}
+
+	bool GetRangeImpl(TArrayView<T> OutValues, int32 Index, const IPCGAttributeAccessorKeys&) const
+	{
+		for (int32 i = 0; i < OutValues.Num(); ++i)
+		{
+			T& OutValue = OutValues[i];
+			const int32 DataIndex = Index + i;
+
+			if constexpr (bUseCStr && Extent == 1 && SubExtent == 1) // String
+			{
+				OutValue = (*SamplePtr)[DataIndex].c_str();
+			}
+			else if constexpr (Extent == 1 && SubExtent == 1) // Scalar
+			{
+				OutValue = (*SamplePtr)[DataIndex];
+			}
+			else if constexpr (Extent > 1 && SubExtent == 1)
+			{
+				for (int32 D = 0; D < Extent; ++D)
+				{
+					OutValue[D] = (*SamplePtr)[DataIndex][D];
+				}
+			}
+			else if constexpr (Extent == 1 && SubExtent > 1)
+			{
+				for (int32 D = 0; D < SubExtent; ++D)
+				{
+					OutValue[D] = (*SamplePtr)[DataIndex * SubExtent + D];
+				}
+			}
+			else
+			{
+				return false;
+			}
+		}
+
+		return true;
+	}
+
+	bool SetRangeImpl(TArrayView<const T>, int32, IPCGAttributeAccessorKeys&, EPCGAttributeAccessorFlags)
+	{
+		return false;
+	}
+
+private:
+	AbcParamType Param;
+	using AbcSamplerType = typename AbcParamType::Sample::samp_ptr_type;
+	AbcSamplerType SamplePtr;
+};
+
+TUniquePtr<const IPCGAttributeAccessor> CreateAlembicPropAccessor(Alembic::AbcGeom::ICompoundProperty& Parameters, Alembic::Abc::PropertyHeader& PropertyHeader, const FString& PropName)
+{
+	Alembic::Abc::PropertyType PropType = PropertyHeader.getPropertyType();
+
+	Alembic::Abc::DataType DataType = PropertyHeader.getDataType();
+	int TypeExtent = DataType.getExtent();
+
+	FString MetadataExtent(PropertyHeader.getMetaData().get("arrayExtent").c_str());
+	const int32 SubExtent = (MetadataExtent.Compare("") == 0 ? 1 : FCString::Atoi(*MetadataExtent));
+
+	if (DataType.getPod() == Alembic::Util::kFloat32POD && (TypeExtent == 1 || SubExtent == 1))
+	{
+		switch (TypeExtent)
+		{
+		case 1:
+			switch (SubExtent)
+			{
+			case 1: return MakeUnique<FPCGAlembicAccessor<float, Alembic::AbcGeom::IFloatGeomParam, 1, 1, false>>(Parameters, PropName);
+			case 2: return MakeUnique<FPCGAlembicAccessor<FVector2D, Alembic::AbcGeom::IFloatGeomParam, 1, 2, false>>(Parameters, PropName);
+			case 3: return MakeUnique<FPCGAlembicAccessor<FVector, Alembic::AbcGeom::IFloatGeomParam, 1, 3, false>>(Parameters, PropName);
+			case 4: return MakeUnique<FPCGAlembicAccessor<FVector4, Alembic::AbcGeom::IFloatGeomParam, 1, 4, false>>(Parameters, PropName);
+			}
+		case 2: return MakeUnique<FPCGAlembicAccessor<FVector2D, Alembic::AbcGeom::IV2fGeomParam, 2, 1, false>>(Parameters, PropName);
+		case 3: return MakeUnique<FPCGAlembicAccessor<FVector, Alembic::AbcGeom::IV3fGeomParam, 3, 1, false>>(Parameters, PropName);
+		case 4: return MakeUnique<FPCGAlembicAccessor<FVector4, Alembic::AbcGeom::IQuatfGeomParam, 4, 1, false>>(Parameters, PropName);
+		}
+	}
+	else if (DataType.getPod() == Alembic::Util::kFloat64POD && (TypeExtent == 1 || SubExtent == 1))
+	{
+		switch (TypeExtent)
+		{
+		case 1:
+			switch (SubExtent)
+			{
+			case 1: return MakeUnique<FPCGAlembicAccessor<double, Alembic::AbcGeom::IDoubleGeomParam, 1, 1, false>>(Parameters, PropName);
+			case 2: return MakeUnique<FPCGAlembicAccessor<FVector2D, Alembic::AbcGeom::IDoubleGeomParam, 1, 2, false>>(Parameters, PropName);
+			case 3: return MakeUnique<FPCGAlembicAccessor<FVector, Alembic::AbcGeom::IDoubleGeomParam, 1, 3, false>>(Parameters, PropName);
+			case 4: return MakeUnique<FPCGAlembicAccessor<FVector4, Alembic::AbcGeom::IDoubleGeomParam, 1, 4, false>>(Parameters, PropName);
+			}
+		case 2: return MakeUnique<FPCGAlembicAccessor<FVector2D, Alembic::AbcGeom::IV2dGeomParam, 2, 1, false>>(Parameters, PropName);
+		case 3: return MakeUnique<FPCGAlembicAccessor<FVector, Alembic::AbcGeom::IV3dGeomParam, 3, 1, false>>(Parameters, PropName);
+		case 4: return MakeUnique<FPCGAlembicAccessor<FVector4, Alembic::AbcGeom::IQuatdGeomParam, 4, 1, false>>(Parameters, PropName);
+		}
+	}
+	else if ((TypeExtent == 0 || TypeExtent == 1) && SubExtent == 1) // Scalar types
+	{
+		switch (DataType.getPod())
+		{
+		case Alembic::Util::kBooleanPOD: return MakeUnique<FPCGAlembicAccessor<bool, Alembic::AbcGeom::IBoolGeomParam, 1, 1, false>>(Parameters, PropName);
+		case Alembic::Util::kInt8POD: return MakeUnique<FPCGAlembicAccessor<int32, Alembic::AbcGeom::ICharGeomParam, 1, 1, false>>(Parameters, PropName);
+		case Alembic::Util::kInt16POD: return MakeUnique<FPCGAlembicAccessor<int32, Alembic::AbcGeom::IInt16GeomParam, 1, 1, false>>(Parameters, PropName);
+		case Alembic::Util::kInt32POD: return MakeUnique<FPCGAlembicAccessor<int32, Alembic::AbcGeom::IInt32GeomParam, 1, 1, false>>(Parameters, PropName);
+		case Alembic::Util::kInt64POD: return MakeUnique<FPCGAlembicAccessor<int64, Alembic::AbcGeom::IInt64GeomParam, 1, 1, false>>(Parameters, PropName);
+		case Alembic::Util::kUnknownPOD: // fall-through
+		case Alembic::Util::kStringPOD: return MakeUnique<FPCGAlembicAccessor<FString, Alembic::AbcGeom::IStringGeomParam, 1, 1, true>>(Parameters, PropName);
+		}
+	}
+
+	return nullptr;
+}
+
+void ParseAlembicObject(FPCGExternalDataContext* Context, const Alembic::Abc::IObject& Object)
+{
+	check(Context);
+	const UPCGExternalDataSettings* Settings = Context->GetInputSettings<UPCGExternalDataSettings>();
+	check(Settings);
+
+	const Alembic::Abc::MetaData& ObjectMetaData = Object.getMetaData();
+	const uint32 NumChildren = Object.getNumChildren();
+
+	if (Alembic::AbcGeom::IPoints::matches(ObjectMetaData))
+	{
+		Alembic::AbcGeom::IPoints Points = Alembic::AbcGeom::IPoints(Object, Alembic::Abc::kWrapExisting);
+		Alembic::AbcGeom::IPoints::schema_type::Sample Sample = Points.getSchema().getValue();
+
+		Alembic::Abc::P3fArraySamplePtr Positions = Sample.getPositions();
+		uint32 NumPoints = Positions ? Positions->size() : 0;
+
+		if (NumPoints > 0)
+		{
+			// Create point data & mapping
+			UPCGPointData* PointData = NewObject<UPCGPointData>();
+			check(PointData);
+
+			if (UPCGComponent* SourceComponent = Context->SourceComponent.Get())
+			{
+				PointData->TargetActor = SourceComponent->GetOwner();
+			}
+
+			FPCGExternalDataContext::FPointDataAccessorsMapping& PointDataAccessorMapping = Context->PointDataAccessorsMapping.Emplace_GetRef();
+			PointDataAccessorMapping.PointData = PointData;
+			// We're not going to use the input keys, but we still need to provide something
+			PointDataAccessorMapping.RowKeys = MakeUnique<FPCGAttributeAccessorKeysEntries>(PCGInvalidEntryKey);
+
+			TArray<FPCGPoint>& OutPoints = PointData->GetMutablePoints();
+			OutPoints.SetNum(NumPoints);
+
+			// Create points
+			for (uint32 PointIndex = 0; PointIndex < NumPoints; ++PointIndex)
+			{
+				Alembic::Abc::P3fArraySample::value_type Position = (*Positions)[PointIndex];
+				OutPoints[PointIndex].Transform.SetLocation(FVector(Position.x, Position.y, Position.z));
+			}
+
+			UPCGMetadata* PointMetadata = PointData->MutableMetadata();
+
+			Alembic::AbcGeom::ICompoundProperty Parameters = Points.getSchema().getArbGeomParams();
+			for (int Index = 0; Index < Parameters.getNumProperties(); ++Index)
+			{
+				Alembic::Abc::PropertyHeader PropertyHeader = Parameters.getPropertyHeader(Index);
+				FString PropName(PropertyHeader.getName().c_str());
+
+				// We'll parse only properties that affect every point/object
+				if (PropertyHeader.getPropertyType() != Alembic::Abc::kArrayProperty &&
+					!(PropertyHeader.getPropertyType() == Alembic::Abc::kCompoundProperty && (PropertyHeader.getDataType().getPod() == Alembic::Util::kUnknownPOD || PropertyHeader.getDataType().getPod() == Alembic::Util::kStringPOD)))
+				{
+					PCGE_LOG_C(Error, GraphAndLog, Context, FText::Format(LOCTEXT("AlembicPropertyType", "Property '{0}' is not supported (expected array)."), FText::FromString(PropName)));
+					continue;
+				}
+
+				TUniquePtr<const IPCGAttributeAccessor> AlembicPropAccessor = CreateAlembicPropAccessor(Parameters, PropertyHeader, PropName);
+
+				if (!AlembicPropAccessor)
+				{
+					PCGE_LOG_C(Error, GraphAndLog, Context, FText::Format(LOCTEXT("AlembicPropertyNotSupported", "Property '{0}' is not of a supported type."), FText::FromString(PropName)));
+					continue;
+				}
+
+				// Setup attribute property selector
+				FPCGAttributePropertySelector PointPropertySelector;
+				if (const FPCGAttributePropertySelector* MappedField = Settings->AttributeMapping.Find(PropName))
+				{
+					PointPropertySelector = *MappedField;
+					PropName = PointPropertySelector.GetDisplayText().ToString();
+				}
+				else
+				{
+					PointPropertySelector.Update(PropName);
+				}
+
+				// Create attribute if needed
+				if (PointPropertySelector.Selection == EPCGAttributePropertySelection::Attribute && !PointMetadata->HasAttribute(FName(PropName)))
+				{
+					auto CreateAttribute = [PointMetadata, &PropName](auto Dummy)
+					{
+						using AttributeType = decltype(Dummy);
+						return PCGMetadataElementCommon::ClearOrCreateAttribute(PointMetadata, FName(PropName), AttributeType{}) != nullptr;
+					};
+
+					if (!PCGMetadataAttribute::CallbackWithRightType(AlembicPropAccessor->GetUnderlyingType(), CreateAttribute))
+					{
+						PCGE_LOG_C(Warning, GraphAndLog, Context, FText::Format(LOCTEXT("FailedToCreateNewAttribute", "Failed to create new attribute '{0}'"), FText::FromString(PropName)));
+						continue;
+					}
+				}
+
+				TUniquePtr<IPCGAttributeAccessor> PointPropertyAccessor = PCGAttributeAccessorHelpers::CreateAccessor(PointData, PointPropertySelector);
+
+				if (!PointPropertyAccessor)
+				{
+					PCGE_LOG_C(Error, GraphAndLog, Context, FText::Format(LOCTEXT("AlembicTargetNotSupported", "Unable to write to target property '{0}.'"), PointPropertySelector.GetDisplayText()));
+					continue;
+				}
+
+				// Final verification, if we can put the value of input into output
+				if (!PCG::Private::IsBroadcastable(AlembicPropAccessor->GetUnderlyingType(), PointPropertyAccessor->GetUnderlyingType()))
+				{
+					FText InputTypeName = FText::FromString(PCG::Private::GetTypeName(AlembicPropAccessor->GetUnderlyingType()));
+					FText OutputTypeName = FText::FromString(PCG::Private::GetTypeName(PointPropertyAccessor->GetUnderlyingType()));
+
+					PCGE_LOG_C(Error, GraphAndLog, Context, FText::Format(LOCTEXT("CannotBroadcastTypes", "Cannot convert input type '{0}' into output type '{1}'"), InputTypeName, OutputTypeName));
+					continue;
+				}
+
+				PointDataAccessorMapping.RowToPointAccessors.Emplace(MoveTemp(AlembicPropAccessor), MoveTemp(PointPropertyAccessor), PointPropertySelector);
+			}
+		}
+	}
+
+	if (NumChildren > 0)
+	{
+		for (uint32 ChildIndex = 0; ChildIndex < NumChildren; ++ChildIndex)
+		{
+			ParseAlembicObject(Context, Object.getChild(ChildIndex));
+		}
+	}
+}
+
+void LoadFromAlembicFile(FPCGLoadAlembicContext* Context, const FString& FileName)
+{
+	check(Context);
+
+	/** Factory used to generate objects*/
+	Alembic::AbcCoreFactory::IFactory& Factory = Context->Factory;
+	Alembic::AbcCoreFactory::IFactory::CoreType CompressionType = Alembic::AbcCoreFactory::IFactory::kUnknown;
+	/** Archive-typed ABC file */
+	Alembic::Abc::IArchive& Archive = Context->Archive;
+	/** Alembic typed root (top) object*/
+	Alembic::Abc::IObject& TopObject = Context->TopObject;
+
+	Factory.setPolicy(Alembic::Abc::ErrorHandler::kQuietNoopPolicy);
+	const size_t ReasonableNumStreams = 12; // This is what we had in Rule Processor historically
+	Factory.setOgawaNumStreams(ReasonableNumStreams);
+
+	// Extract Archive and compression type from file
+	Archive = Factory.getArchive(TCHAR_TO_UTF8(*FileName), CompressionType);
+
+	if (!Archive.valid())
+	{
+		PCGE_LOG_C(Error, GraphAndLog, Context, FText::Format(LOCTEXT("FailedToOpenAlembic", "Failed to open '{0}': Not a valid Alembic file."), FText::FromString(FileName)));
+		return;
+	}
+
+	// Get Top/root object
+	TopObject = Alembic::Abc::IObject(Archive, Alembic::Abc::kTop);
+	if (!TopObject.valid())
+	{
+		PCGE_LOG_C(Error, GraphAndLog, Context, FText::Format(LOCTEXT("FailedToOpenAlembicRoot", "Failed to import '{0}': Alembic root is not valid."), FText::FromString(FileName)));
+		return;
+	}
+
+	ParseAlembicObject(Context, TopObject);
+}
+
+#endif // WITH_EDITOR
+
+}
+
+#undef LOCTEXT_NAMESPACE
\ No newline at end of file
diff --git a/Engine/Plugins/Experimental/PCGExternalDataInterop/Source/PCGExternalDataInterop/Private/Alembic/PCGAlembicInterop.h b/Engine/Plugins/Experimental/PCGExternalDataInterop/Source/PCGExternalDataInterop/Private/Alembic/PCGAlembicInterop.h
new file mode 100644
--- /dev/null
+++ b/Engine/Plugins/Experimental/PCGExternalDataInterop/Source/PCGExternalDataInterop/Private/Alembic/PCGAlembicInterop.h
@@ -0,0 +1,42 @@
+// Copyright Epic Games, Inc. All Rights Reserved.
+
+#pragma once
+
+#include "Elements/IO/PCGExternalDataContext.h"
+
+class UPCGPointData;
+
+#if WITH_EDITOR
+#include "AbcImportSettings.h"
+#include "AbcUtilities.h"
+
+THIRD_PARTY_INCLUDES_START
+#pragma warning(push)
+#pragma warning(disable:4005) // TEXT macro redefinition
+#include <Alembic/AbcCoreOgawa/All.h>
+#include "Alembic/AbcGeom/All.h"
+#include "Alembic/AbcCoreFactory/IFactory.h"
+#pragma warning(pop)
+THIRD_PARTY_INCLUDES_END
+
+#endif // WITH_EDITOR
+
+struct FPCGLoadAlembicContext : public FPCGExternalDataContext
+{
+#if WITH_EDITOR
+	// Used for alembic parsing
+	/** Factory used to generate objects*/
+	Alembic::AbcCoreFactory::IFactory Factory;
+	/** Archive-typed ABC file */
+	Alembic::Abc::IArchive Archive;
+	/** Alembic typed root (top) object*/
+	Alembic::Abc::IObject TopObject;
+#endif // WITH_EDITOR
+};
+
+namespace PCGAlembicInterop
+{
+#if WITH_EDITOR
+	void LoadFromAlembicFile(FPCGLoadAlembicContext* Context, const FString& FileName);
+#endif
+}
\ No newline at end of file
diff --git a/Engine/Plugins/Experimental/PCGExternalDataInterop/Source/PCGExternalDataInterop/Private/Elements/PCGLoadAlembicElement.cpp b/Engine/Plugins/Experimental/PCGExternalDataInterop/Source/PCGExternalDataInterop/Private/Elements/PCGLoadAlembicElement.cpp
new file mode 100644
--- /dev/null
+++ b/Engine/Plugins/Experimental/PCGExternalDataInterop/Source/PCGExternalDataInterop/Private/Elements/PCGLoadAlembicElement.cpp
@@ -0,0 +1,108 @@
+// Copyright Epic Games, Inc. All Rights Reserved.
+
+#include "Elements/PCGLoadAlembicElement.h"
+
+#include "Alembic/PCGAlembicInterop.h"
+#include "Misc/Paths.h"
+
+#define LOCTEXT_NAMESPACE "PCGLoadAlembic"
+
+#if WITH_EDITOR
+FText UPCGLoadAlembicSettings::GetDefaultNodeTitle() const
+{
+	return LOCTEXT("NodeTitle", "Load Alembic");
+}
+
+FText UPCGLoadAlembicSettings::GetNodeTooltipText() const
+{
+	return LOCTEXT("NodeTooltip", "Loads data from an Alembic file");
+}
+#endif // WITH_EDITOR
+
+FPCGElementPtr UPCGLoadAlembicSettings::CreateElement() const
+{
+	return MakeShared<FPCGLoadAlembicElement>();
+}
+
+FPCGContext* FPCGLoadAlembicElement::Initialize(const FPCGDataCollection& InputData, TWeakObjectPtr<UPCGComponent> SourceComponent, const UPCGNode* Node)
+{
+	FPCGLoadAlembicContext* Context = new FPCGLoadAlembicContext();
+	Context->InputData = InputData;
+	Context->SourceComponent = SourceComponent;
+	Context->Node = Node;
+
+	return Context;
+}
+
+bool FPCGLoadAlembicElement::PrepareLoad(FPCGExternalDataContext* InContext) const
+{
+	check(InContext);
+	FPCGLoadAlembicContext* Context = static_cast<FPCGLoadAlembicContext*>(InContext);
+
+	check(Context);
+	const UPCGLoadAlembicSettings* Settings = Context->GetInputSettings<UPCGLoadAlembicSettings>();
+	check(Settings);
+
+#if WITH_EDITOR
+	const FString FileName = Settings->AlembicFilePath.FilePath;
+	PCGAlembicInterop::LoadFromAlembicFile(Context, FileName);
+
+	if (!Context->PointDataAccessorsMapping.IsEmpty())
+	{
+		for (const FPCGExternalDataContext::FPointDataAccessorsMapping& DataMapping : Context->PointDataAccessorsMapping)
+		{
+			FPCGTaggedData& OutData = Context->OutputData.TaggedData.Emplace_GetRef();
+			OutData.Data = DataMapping.PointData;
+		}
+
+		Context->bDataPrepared = true;
+	}
+#else
+	PCGE_LOG(Error, GraphAndLog, LOCTEXT("NotSupportedInGameMode", "The Load Alembic node is not support in non-editor builds."));
+#endif
+
+	return true;
+}
+
+bool FPCGLoadAlembicElement::ExecuteLoad(FPCGExternalDataContext* InContext) const
+{
+	if (!FPCGExternalDataElement::ExecuteLoad(InContext))
+	{
+		return false;
+	}
+
+	// Finally, apply conversion
+	check(InContext);
+	FPCGLoadAlembicContext* Context = static_cast<FPCGLoadAlembicContext*>(InContext);
+	check(Context);
+	const UPCGLoadAlembicSettings* Settings = Context->GetInputSettings<UPCGLoadAlembicSettings>();
+	check(Settings);
+
+	const FVector& ConversionScale = Settings->ConversionScale;
+	const FVector& ConversionRotation = Settings->ConversionRotation;
+
+	const FTransform ConversionTransform(FRotator::MakeFromEuler(ConversionRotation), FVector::ZeroVector, ConversionScale);
+	if (!ConversionTransform.Equals(FTransform::Identity))
+	{
+		for (const FPCGExternalDataContext::FPointDataAccessorsMapping& DataMapping : Context->PointDataAccessorsMapping)
+		{
+			UPCGPointData* PointData = DataMapping.PointData;
+
+			if (!PointData)
+			{
+				continue;
+			}
+
+			TArray<FPCGPoint>& Points = PointData->GetMutablePoints();
+
+			for (FPCGPoint& Point : Points)
+			{
+				Point.Transform = Point.Transform * ConversionTransform;
+			}
+		}
+	}
+
+	return true;
+}
+
+#undef LOCTEXT_NAMESPACE
\ No newline at end of file
diff --git a/Engine/Plugins/Experimental/PCGExternalDataInterop/Source/PCGExternalDataInterop/Private/PCGExternalDataInteropModule.cpp b/Engine/Plugins/Experimental/PCGExternalDataInterop/Source/PCGExternalDataInterop/Private/PCGExternalDataInteropModule.cpp
new file mode 100644
--- /dev/null
+++ b/Engine/Plugins/Experimental/PCGExternalDataInterop/Source/PCGExternalDataInterop/Private/PCGExternalDataInteropModule.cpp
@@ -0,0 +1,16 @@
+// Copyright Epic Games, Inc. All Rights Reserved.
+
+#include "Modules/ModuleManager.h"
+
+class FPCGExternalDataInteropModule final : public IModuleInterface
+{
+public:
+	//~ IModuleInterface implementation
+	virtual bool SupportsDynamicReloading() override
+	{
+		return true;
+	}
+	//~ End IModuleInterface implementation
+};
+
+IMPLEMENT_MODULE(FPCGExternalDataInteropModule, PCGExternalDataInterop);
diff --git a/Engine/Plugins/Experimental/PCGExternalDataInterop/Source/PCGExternalDataInterop/Public/Elements/PCGLoadAlembicElement.h b/Engine/Plugins/Experimental/PCGExternalDataInterop/Source/PCGExternalDataInterop/Public/Elements/PCGLoadAlembicElement.h
new file mode 100644
--- /dev/null
+++ b/Engine/Plugins/Experimental/PCGExternalDataInterop/Source/PCGExternalDataInterop/Public/Elements/PCGLoadAlembicElement.h
@@ -0,0 +1,50 @@
+// Copyright Epic Games, Inc. All Rights Reserved.
+
+#pragma once
+
+#include "Elements/IO/PCGExternalData.h"
+
+#include "PCGLoadAlembicElement.generated.h"
+
+struct FPCGExternalDataContext;
+
+UCLASS(BlueprintType, ClassGroup = (Procedural))
+class PCGEXTERNALDATAINTEROP_API UPCGLoadAlembicSettings : public UPCGExternalDataSettings
+{
+	GENERATED_BODY()
+
+public:
+	// ~Begin UPCGSettings interface
+#if WITH_EDITOR
+	virtual FName GetDefaultNodeName() const override { return FName(TEXT("LoadAlembic")); }
+	virtual FText GetDefaultNodeTitle() const override;
+	virtual FText GetNodeTooltipText() const override;
+#endif
+
+protected:
+	virtual FPCGElementPtr CreateElement() const override;
+	// ~End UPCGSettings interface
+
+public:
+	UPROPERTY(BlueprintReadWrite, EditAnywhere, Category = "Alembic", meta = (FilePathFilter = "Alembic files (*.abc)|*.abc", PCG_Overridable))
+	FFilePath AlembicFilePath;
+
+	// To prevent a dependency on the alembic editor module in this class, we'll keep around only the types we need
+	/** Scale to apply during import. Note that for both Max/Maya presets the value flips the Y axis. */
+	UPROPERTY(BlueprintReadWrite, EditAnywhere, Category = "Alembic", meta = (PCG_Overridable))
+	FVector ConversionScale = FVector(1.0f, -1.0f, 1.0f);
+
+	/** Rotation in Euler angles applied during import. For Max, use (90, 0, 0). */
+	UPROPERTY(BlueprintReadWrite, EditAnywhere, Category = "Alembic", meta = (PCG_Overridable))
+	FVector ConversionRotation = FVector::ZeroVector;
+};
+
+class PCGEXTERNALDATAINTEROP_API FPCGLoadAlembicElement : public FPCGExternalDataElement
+{
+public:
+	virtual FPCGContext* Initialize(const FPCGDataCollection& InputData, TWeakObjectPtr<UPCGComponent> SourceComponent, const UPCGNode* Node) override;
+
+protected:
+	virtual bool PrepareLoad(FPCGExternalDataContext* Context) const override;
+	virtual bool ExecuteLoad(FPCGExternalDataContext* Context) const override;
+};
\ No newline at end of file
diff --git a/Engine/Plugins/Experimental/PCGExternalDataInterop/Source/PCGExternalDataInteropEditor/PCGExternalDataInteropEditor.Build.cs b/Engine/Plugins/Experimental/PCGExternalDataInterop/Source/PCGExternalDataInteropEditor/PCGExternalDataInteropEditor.Build.cs
new file mode 100644
--- /dev/null
+++ b/Engine/Plugins/Experimental/PCGExternalDataInterop/Source/PCGExternalDataInteropEditor/PCGExternalDataInteropEditor.Build.cs
@@ -0,0 +1,34 @@
+// Copyright Epic Games, Inc. All Rights Reserved.
+
+namespace UnrealBuildTool.Rules
+{
+	public class PCGExternalDataInteropEditor : ModuleRules
+	{
+		public PCGExternalDataInteropEditor(ReadOnlyTargetRules Target) : base(Target)
+		{
+			PCHUsage = PCHUsageMode.UseExplicitOrSharedPCHs;
+
+			PublicDependencyModuleNames.AddRange(
+				new string[] {
+					"Core",
+					"CoreUObject",
+					"Engine",
+					"Projects",
+					"PCG"
+				});
+
+			if (Target.bBuildEditor == true)
+			{
+				PrivateDependencyModuleNames.AddRange(
+					new string[]
+					{
+						"PCGExternalDataInterop",
+						"UnrealEd",
+						"AlembicLib",
+						"AlembicLibrary"
+					}
+				);
+			}
+		}
+	}
+}
diff --git a/Engine/Plugins/Experimental/PCGExternalDataInterop/Source/PCGExternalDataInteropEditor/Private/PCGExternalDataInteropEditorModule.cpp b/Engine/Plugins/Experimental/PCGExternalDataInterop/Source/PCGExternalDataInteropEditor/Private/PCGExternalDataInteropEditorModule.cpp
new file mode 100644
--- /dev/null
+++ b/Engine/Plugins/Experimental/PCGExternalDataInterop/Source/PCGExternalDataInteropEditor/Private/PCGExternalDataInteropEditorModule.cpp
@@ -0,0 +1,16 @@
+// Copyright Epic Games, Inc. All Rights Reserved.
+
+#include "Modules/ModuleManager.h"
+
+class FPCGExternalDataInteropEditorModule final : public IModuleInterface
+{
+public:
+	//~ IModuleInterface implementation
+	virtual bool SupportsDynamicReloading() override
+	{
+		return true;
+	}
+	//~ End IModuleInterface implementation
+};
+
+IMPLEMENT_MODULE(FPCGExternalDataInteropEditorModule, PCGExternalDataInteropEditor);
diff --git a/Engine/Plugins/Experimental/PCGExternalDataInterop/Source/PCGExternalDataInteropEditor/Private/PCGLoadAlembic.cpp b/Engine/Plugins/Experimental/PCGExternalDataInterop/Source/PCGExternalDataInteropEditor/Private/PCGLoadAlembic.cpp
new file mode 100644
--- /dev/null
+++ b/Engine/Plugins/Experimental/PCGExternalDataInterop/Source/PCGExternalDataInteropEditor/Private/PCGLoadAlembic.cpp
@@ -0,0 +1,41 @@
+// Copyright Epic Games, Inc. All Rights Reserved.
+
+#include "PCGLoadAlembic.h"
+
+#include "Elements/PCGLoadAlembicElement.h"
+
+void UPCGLoadAlembicFunctionLibrary::LoadAlembicFileToPCG(const FPCGLoadAlembicBPData& InSettings, FPCGDataCollection& Data, UObject* TargetOuter)
+{
+	// Simulate execution of the FPCGLoadAlembicElement.
+	UPCGLoadAlembicSettings* Settings = NewObject<UPCGLoadAlembicSettings>();
+	Settings->AlembicFilePath = InSettings.AlembicFilePath;
+	Settings->ConversionScale = InSettings.ConversionSettings.Scale;
+	Settings->ConversionRotation = InSettings.ConversionSettings.Rotation;
+	Settings->AttributeMapping = InSettings.AttributeMapping;
+
+	FPCGLoadAlembicElement Element;
+	FPCGContext* Context = Element.Initialize(FPCGDataCollection(), TWeakObjectPtr<UPCGComponent>(), nullptr);
+	Context->AsyncState.NumAvailableTasks = 1; // TODO : change this
+	FPCGTaggedData& DataSettings = Context->InputData.TaggedData.Emplace_GetRef();
+	DataSettings.Data = Settings;
+
+	while (!Element.Execute(Context))
+	{
+		// Nothing
+	}
+
+	Data = Context->OutputData;
+	delete Context;
+
+	// Reouter the data if needed
+	if (TargetOuter)
+	{
+		for (FPCGTaggedData& TaggedData : Data.TaggedData)
+		{
+			if (TaggedData.Data)
+			{
+				Cast<UObject>(TaggedData.Data)->Rename(nullptr, TargetOuter, REN_ForceNoResetLoaders | REN_DoNotDirty | REN_DontCreateRedirectors);
+			}
+		}
+	}
+}
\ No newline at end of file
diff --git a/Engine/Plugins/Experimental/PCGExternalDataInterop/Source/PCGExternalDataInteropEditor/Public/PCGLoadAlembic.h b/Engine/Plugins/Experimental/PCGExternalDataInterop/Source/PCGExternalDataInteropEditor/Public/PCGLoadAlembic.h
new file mode 100644
--- /dev/null
+++ b/Engine/Plugins/Experimental/PCGExternalDataInterop/Source/PCGExternalDataInteropEditor/Public/PCGLoadAlembic.h
@@ -0,0 +1,34 @@
+// Copyright Epic Games, Inc. All Rights Reserved.
+
+#pragma once
+
+#include "Metadata/PCGAttributePropertySelector.h"
+
+#include "AbcImportSettings.h"
+
+#include "PCGLoadAlembic.generated.h"
+
+USTRUCT(BlueprintType)
+struct PCGEXTERNALDATAINTEROPEDITOR_API FPCGLoadAlembicBPData
+{
+	GENERATED_BODY()
+
+	UPROPERTY(EditAnywhere, Category = "Alembic", meta = (FilePathFilter = "Alembic files (*.abc)|*.abc", PCG_Overridable))
+	FFilePath AlembicFilePath;
+
+	/** Conversion settings that will be applied on the transform only */
+	UPROPERTY(BlueprintReadWrite, EditAnywhere, Category = "Settings")
+	FAbcConversionSettings ConversionSettings;
+
+	UPROPERTY(BlueprintReadWrite, EditAnywhere, Category = "Settings")
+	TMap<FString, FPCGAttributePropertySelector> AttributeMapping;
+};
+
+UCLASS()
+class PCGEXTERNALDATAINTEROPEDITOR_API UPCGLoadAlembicFunctionLibrary : public UBlueprintFunctionLibrary
+{
+	GENERATED_BODY()
+
+	UFUNCTION(BlueprintCallable, Category = "PCG|IO")
+	static void LoadAlembicFileToPCG(const FPCGLoadAlembicBPData& Settings, FPCGDataCollection& Data, UObject* TargetOuter);
+};
\ No newline at end of file
