Tentative build fix around weakptrs in the data interop elements

#rnx

diff --git a/Engine/Plugins/Experimental/PCG/Source/PCG/Private/Elements/IO/PCGDataTableElement.cpp b/Engine/Plugins/Experimental/PCG/Source/PCG/Private/Elements/IO/PCGDataTableElement.cpp
--- a/Engine/Plugins/Experimental/PCG/Source/PCG/Private/Elements/IO/PCGDataTableElement.cpp
+++ b/Engine/Plugins/Experimental/PCG/Source/PCG/Private/Elements/IO/PCGDataTableElement.cpp
@@ -63,11 +63,6 @@ bool FPCGLoadDataTableElement::PrepareLoad(FPCGExternalDataContext* Context) con
 	UPCGPointData* PointData = NewObject<UPCGPointData>();
 	check(PointData);
 
-	if (Context->SourceComponent.Get())
-	{
-		PointData->TargetActor = Context->SourceComponent->GetOwner();
-	}
-
 	FPCGExternalDataContext::FPointDataAccessorsMapping& PointDataAccessorMapping = Context->PointDataAccessorsMapping.Emplace_GetRef();
 	PointDataAccessorMapping.PointData = PointData;
 
diff --git a/Engine/Plugins/Experimental/PCGExternalDataInterop/Source/PCGExternalDataInterop/Private/Alembic/PCGAlembicInterop.cpp b/Engine/Plugins/Experimental/PCGExternalDataInterop/Source/PCGExternalDataInterop/Private/Alembic/PCGAlembicInterop.cpp
--- a/Engine/Plugins/Experimental/PCGExternalDataInterop/Source/PCGExternalDataInterop/Private/Alembic/PCGAlembicInterop.cpp
+++ b/Engine/Plugins/Experimental/PCGExternalDataInterop/Source/PCGExternalDataInterop/Private/Alembic/PCGAlembicInterop.cpp
@@ -165,11 +165,6 @@ void ParseAlembicObject(FPCGExternalDataContext* Context, const Alembic::Abc::IO
 			UPCGPointData* PointData = NewObject<UPCGPointData>();
 			check(PointData);
 
-			if (UPCGComponent* SourceComponent = Context->SourceComponent.Get())
-			{
-				PointData->TargetActor = SourceComponent->GetOwner();
-			}
-
 			FPCGExternalDataContext::FPointDataAccessorsMapping& PointDataAccessorMapping = Context->PointDataAccessorsMapping.Emplace_GetRef();
 			PointDataAccessorMapping.PointData = PointData;
 			// We're not going to use the input keys, but we still need to provide something
