Tentative fix for static analysis warning wrt to include after first block

#rnx

diff --git a/Engine/Plugins/Experimental/PCGExternalDataInterop/Source/PCGExternalDataInterop/Private/Alembic/PCGAlembicInterop.h b/Engine/Plugins/Experimental/PCGExternalDataInterop/Source/PCGExternalDataInterop/Private/Alembic/PCGAlembicInterop.h
--- a/Engine/Plugins/Experimental/PCGExternalDataInterop/Source/PCGExternalDataInterop/Private/Alembic/PCGAlembicInterop.h
+++ b/Engine/Plugins/Experimental/PCGExternalDataInterop/Source/PCGExternalDataInterop/Private/Alembic/PCGAlembicInterop.h
@@ -4,8 +4,6 @@
 
 #include "Elements/IO/PCGExternalDataContext.h"
 
-class UPCGPointData;
-
 #if WITH_EDITOR
 #include "AbcImportSettings.h"
 #include "AbcUtilities.h"
