Added support to remap the positions in an Alembic file parsing

diff --git a/Engine/Plugins/Experimental/PCGExternalDataInterop/Source/PCGExternalDataInterop/Private/Alembic/PCGAlembicInterop.cpp b/Engine/Plugins/Experimental/PCGExternalDataInterop/Source/PCGExternalDataInterop/Private/Alembic/PCGAlembicInterop.cpp
--- a/Engine/Plugins/Experimental/PCGExternalDataInterop/Source/PCGExternalDataInterop/Private/Alembic/PCGAlembicInterop.cpp
+++ b/Engine/Plugins/Experimental/PCGExternalDataInterop/Source/PCGExternalDataInterop/Private/Alembic/PCGAlembicInterop.cpp
@@ -19,6 +19,37 @@ namespace PCGAlembicInterop
 {
 
 #if WITH_EDITOR
+class FPCGAlembicPositionsAccessor : public IPCGAttributeAccessorT<FPCGAlembicPositionsAccessor>
+{
+public:
+	using Type = FVector;
+	using Super = IPCGAttributeAccessorT<FPCGAlembicPositionsAccessor>;
+
+	FPCGAlembicPositionsAccessor(Alembic::AbcGeom::IPoints::schema_type::Sample& Sample)
+		: Super(/*bInReadOnly=*/true)
+		, SamplePtr(Sample.getPositions())
+	{}
+
+	bool GetRangeImpl(TArrayView<FVector> OutValues, int32 Index, const IPCGAttributeAccessorKeys&) const
+	{
+		for (int32 i = 0; i < OutValues.Num(); ++i)
+		{
+			Alembic::Abc::P3fArraySample::value_type Position = (*SamplePtr)[Index + i];
+			OutValues[i] = FVector(Position.x, Position.y, Position.z);
+		}
+
+		return true;
+	}
+
+	bool SetRangeImpl(TArrayView<const FVector>, int32, IPCGAttributeAccessorKeys&, EPCGAttributeAccessorFlags)
+	{
+		return false;
+	}
+
+private:
+	Alembic::Abc::P3fArraySamplePtr SamplePtr;
+};
+
 template<typename T, typename AbcParamType, int32 Extent, int32 SubExtent, bool bUseCStr>
 class FPCGAlembicAccessor : public IPCGAttributeAccessorT<FPCGAlembicAccessor<T, AbcParamType, Extent, SubExtent, bUseCStr>>
 {
@@ -142,6 +173,58 @@ TUniquePtr<const IPCGAttributeAccessor> CreateAlembicPropAccessor(Alembic::AbcGe
 	return nullptr;
 }
 
+bool CreatePointAccessorAndValidate(FPCGContext* Context, UPCGPointData* PointData, const TUniquePtr<const IPCGAttributeAccessor>& AlembicPropAccessor, const FPCGAttributePropertySelector& PointPropertySelector, const FString& PropName, TUniquePtr<IPCGAttributeAccessor>& PointPropertyAccessor)
+{
+	if (!PointData || !PointData->Metadata)
+	{
+		return false;
+	}
+
+	if (!AlembicPropAccessor)
+	{
+		PCGE_LOG_C(Error, GraphAndLog, Context, FText::Format(LOCTEXT("AlembicPropertyNotSupported", "Property '{0}' is not of a supported type."), FText::FromString(PropName)));
+		return false;
+	}
+
+	UPCGMetadata* PointMetadata = PointData->Metadata;
+
+	// Create attribute if needed
+	if (PointPropertySelector.Selection == EPCGAttributePropertySelection::Attribute && !PointMetadata->HasAttribute(FName(PropName)))
+	{
+		auto CreateAttribute = [PointMetadata, &PropName](auto Dummy)
+		{
+			using AttributeType = decltype(Dummy);
+			return PCGMetadataElementCommon::ClearOrCreateAttribute(PointMetadata, FName(PropName), AttributeType{}) != nullptr;
+		};
+
+		if (!PCGMetadataAttribute::CallbackWithRightType(AlembicPropAccessor->GetUnderlyingType(), CreateAttribute))
+		{
+			PCGE_LOG_C(Warning, GraphAndLog, Context, FText::Format(LOCTEXT("FailedToCreateNewAttribute", "Failed to create new attribute '{0}'"), FText::FromString(PropName)));
+			return false;
+		}
+	}
+
+	PointPropertyAccessor = PCGAttributeAccessorHelpers::CreateAccessor(PointData, PointPropertySelector);
+
+	if (!PointPropertyAccessor)
+	{
+		PCGE_LOG_C(Error, GraphAndLog, Context, FText::Format(LOCTEXT("AlembicTargetNotSupported", "Unable to write to target property '{0}.'"), PointPropertySelector.GetDisplayText()));
+		return false;
+	}
+
+	// Final verification, if we can put the value of input into output
+	if (!PCG::Private::IsBroadcastable(AlembicPropAccessor->GetUnderlyingType(), PointPropertyAccessor->GetUnderlyingType()))
+	{
+		FText InputTypeName = FText::FromString(PCG::Private::GetTypeName(AlembicPropAccessor->GetUnderlyingType()));
+		FText OutputTypeName = FText::FromString(PCG::Private::GetTypeName(PointPropertyAccessor->GetUnderlyingType()));
+
+		PCGE_LOG_C(Error, GraphAndLog, Context, FText::Format(LOCTEXT("CannotBroadcastTypes", "Cannot convert input type '{0}' into output type '{1}'"), InputTypeName, OutputTypeName));
+		return false;
+	}
+
+	return true;
+}
+
 void ParseAlembicObject(FPCGExternalDataContext* Context, const Alembic::Abc::IObject& Object)
 {
 	check(Context);
@@ -165,6 +248,8 @@ void ParseAlembicObject(FPCGExternalDataContext* Context, const Alembic::Abc::IO
 			UPCGPointData* PointData = NewObject<UPCGPointData>();
 			check(PointData);
 
+			UPCGMetadata* PointMetadata = PointData->MutableMetadata();
+
 			FPCGExternalDataContext::FPointDataAccessorsMapping& PointDataAccessorMapping = Context->PointDataAccessorsMapping.Emplace_GetRef();
 			PointDataAccessorMapping.PointData = PointData;
 			// We're not going to use the input keys, but we still need to provide something
@@ -173,14 +258,28 @@ void ParseAlembicObject(FPCGExternalDataContext* Context, const Alembic::Abc::IO
 			TArray<FPCGPoint>& OutPoints = PointData->GetMutablePoints();
 			OutPoints.SetNum(NumPoints);
 
-			// Create points
-			for (uint32 PointIndex = 0; PointIndex < NumPoints; ++PointIndex)
+			// If the user has provided a position remapping, don't process points right now, instead push the transformation to the row accessors
+			if (const FPCGAttributePropertySelector* RemappedPositions = Settings->AttributeMapping.Find(TEXT("Position")))
 			{
-				Alembic::Abc::P3fArraySample::value_type Position = (*Positions)[PointIndex];
-				OutPoints[PointIndex].Transform.SetLocation(FVector(Position.x, Position.y, Position.z));
-			}
+				TUniquePtr<const IPCGAttributeAccessor> AlembicPositionAccessor = MakeUnique<FPCGAlembicPositionsAccessor>(Sample);
+				TUniquePtr<IPCGAttributeAccessor> PointPropertyAccessor;
 
-			UPCGMetadata* PointMetadata = PointData->MutableMetadata();
+				FPCGAttributePropertySelector PointPositionSelector = *RemappedPositions;
+				const FString PropName = PointPositionSelector.GetDisplayText().ToString();
+
+				if (CreatePointAccessorAndValidate(Context, PointData, AlembicPositionAccessor, PointPositionSelector, PropName, PointPropertyAccessor))
+				{
+					PointDataAccessorMapping.RowToPointAccessors.Emplace(MoveTemp(AlembicPositionAccessor), MoveTemp(PointPropertyAccessor), PointPositionSelector);
+				}
+			}
+			else // Otherwise, write the positions directly
+			{
+				for (uint32 PointIndex = 0; PointIndex < NumPoints; ++PointIndex)
+				{
+					Alembic::Abc::P3fArraySample::value_type Position = (*Positions)[PointIndex];
+					OutPoints[PointIndex].Transform.SetLocation(FVector(Position.x, Position.y, Position.z));
+				}
+			}
 
 			Alembic::AbcGeom::ICompoundProperty Parameters = Points.getSchema().getArbGeomParams();
 			for (int Index = 0; Index < Parameters.getNumProperties(); ++Index)
@@ -197,12 +296,7 @@ void ParseAlembicObject(FPCGExternalDataContext* Context, const Alembic::Abc::IO
 				}
 
 				TUniquePtr<const IPCGAttributeAccessor> AlembicPropAccessor = CreateAlembicPropAccessor(Parameters, PropertyHeader, PropName);
-
-				if (!AlembicPropAccessor)
-				{
-					PCGE_LOG_C(Error, GraphAndLog, Context, FText::Format(LOCTEXT("AlembicPropertyNotSupported", "Property '{0}' is not of a supported type."), FText::FromString(PropName)));
-					continue;
-				}
+				TUniquePtr<IPCGAttributeAccessor> PointPropertyAccessor;
 
 				// Setup attribute property selector
 				FPCGAttributePropertySelector PointPropertySelector;
@@ -216,41 +310,10 @@ void ParseAlembicObject(FPCGExternalDataContext* Context, const Alembic::Abc::IO
 					PointPropertySelector.Update(PropName);
 				}
 
-				// Create attribute if needed
-				if (PointPropertySelector.Selection == EPCGAttributePropertySelection::Attribute && !PointMetadata->HasAttribute(FName(PropName)))
+				if (CreatePointAccessorAndValidate(Context, PointData, AlembicPropAccessor, PointPropertySelector, PropName, PointPropertyAccessor))
 				{
-					auto CreateAttribute = [PointMetadata, &PropName](auto Dummy)
-					{
-						using AttributeType = decltype(Dummy);
-						return PCGMetadataElementCommon::ClearOrCreateAttribute(PointMetadata, FName(PropName), AttributeType{}) != nullptr;
-					};
-
-					if (!PCGMetadataAttribute::CallbackWithRightType(AlembicPropAccessor->GetUnderlyingType(), CreateAttribute))
-					{
-						PCGE_LOG_C(Warning, GraphAndLog, Context, FText::Format(LOCTEXT("FailedToCreateNewAttribute", "Failed to create new attribute '{0}'"), FText::FromString(PropName)));
-						continue;
-					}
+					PointDataAccessorMapping.RowToPointAccessors.Emplace(MoveTemp(AlembicPropAccessor), MoveTemp(PointPropertyAccessor), PointPropertySelector);
 				}
-
-				TUniquePtr<IPCGAttributeAccessor> PointPropertyAccessor = PCGAttributeAccessorHelpers::CreateAccessor(PointData, PointPropertySelector);
-
-				if (!PointPropertyAccessor)
-				{
-					PCGE_LOG_C(Error, GraphAndLog, Context, FText::Format(LOCTEXT("AlembicTargetNotSupported", "Unable to write to target property '{0}.'"), PointPropertySelector.GetDisplayText()));
-					continue;
-				}
-
-				// Final verification, if we can put the value of input into output
-				if (!PCG::Private::IsBroadcastable(AlembicPropAccessor->GetUnderlyingType(), PointPropertyAccessor->GetUnderlyingType()))
-				{
-					FText InputTypeName = FText::FromString(PCG::Private::GetTypeName(AlembicPropAccessor->GetUnderlyingType()));
-					FText OutputTypeName = FText::FromString(PCG::Private::GetTypeName(PointPropertyAccessor->GetUnderlyingType()));
-
-					PCGE_LOG_C(Error, GraphAndLog, Context, FText::Format(LOCTEXT("CannotBroadcastTypes", "Cannot convert input type '{0}' into output type '{1}'"), InputTypeName, OutputTypeName));
-					continue;
-				}
-
-				PointDataAccessorMapping.RowToPointAccessors.Emplace(MoveTemp(AlembicPropAccessor), MoveTemp(PointPropertyAccessor), PointPropertySelector);
 			}
 		}
 	}
