[PCG] Tentative fix for build break

#rnx

diff --git a/Engine/Plugins/Experimental/PCGExternalDataInterop/Source/PCGExternalDataInterop/Private/Alembic/PCGAlembicInterop.cpp b/Engine/Plugins/Experimental/PCGExternalDataInterop/Source/PCGExternalDataInterop/Private/Alembic/PCGAlembicInterop.cpp
--- a/Engine/Plugins/Experimental/PCGExternalDataInterop/Source/PCGExternalDataInterop/Private/Alembic/PCGAlembicInterop.cpp
+++ b/Engine/Plugins/Experimental/PCGExternalDataInterop/Source/PCGExternalDataInterop/Private/Alembic/PCGAlembicInterop.cpp
@@ -189,7 +189,7 @@ bool CreatePointAccessorAndValidate(FPCGContext* Context, UPCGPointData* PointDa
 	UPCGMetadata* PointMetadata = PointData->Metadata;
 
 	// Create attribute if needed
-	if (PointPropertySelector.Selection == EPCGAttributePropertySelection::Attribute && !PointMetadata->HasAttribute(FName(PropName)))
+	if (PointPropertySelector.GetSelection() == EPCGAttributePropertySelection::Attribute && !PointMetadata->HasAttribute(FName(PropName)))
 	{
 		auto CreateAttribute = [PointMetadata, &PropName](auto Dummy)
 		{
@@ -264,7 +264,8 @@ void ParseAlembicObject(FPCGExternalDataContext* Context, const Alembic::Abc::IO
 				TUniquePtr<const IPCGAttributeAccessor> AlembicPositionAccessor = MakeUnique<FPCGAlembicPositionsAccessor>(Sample);
 				TUniquePtr<IPCGAttributeAccessor> PointPropertyAccessor;
 
-				FPCGAttributePropertySelector PointPositionSelector = *RemappedPositions;
+				FPCGAttributePropertyOutputSelector PointPositionSelector;
+				PointPositionSelector.ImportFromOtherSelector(*RemappedPositions);
 				const FString PropName = PointPositionSelector.GetDisplayText().ToString();
 
 				if (CreatePointAccessorAndValidate(Context, PointData, AlembicPositionAccessor, PointPositionSelector, PropName, PointPropertyAccessor))
