Added handedness flip option for the Alembic import.
Added a CitySample preset so it's easy to use out of the box, mirrored in the EUW as well.
Fixed attribute access name so we can use swizzling on newly created attributes too.

diff --git a/Engine/Plugins/Experimental/PCG/Source/PCG/Private/Elements/IO/PCGDataTableElement.cpp b/Engine/Plugins/Experimental/PCG/Source/PCG/Private/Elements/IO/PCGDataTableElement.cpp
--- a/Engine/Plugins/Experimental/PCG/Source/PCG/Private/Elements/IO/PCGDataTableElement.cpp
+++ b/Engine/Plugins/Experimental/PCG/Source/PCG/Private/Elements/IO/PCGDataTableElement.cpp
@@ -88,7 +88,7 @@ bool FPCGLoadDataTableElement::PrepareLoad(FPCGExternalDataContext* Context) con
 		if (const FPCGAttributePropertySelector* MappedField = Settings->AttributeMapping.Find(FieldName))
 		{
 			PointPropertySelector.ImportFromOtherSelector(*MappedField);
-			FieldName = PointPropertySelector.GetDisplayText().ToString();
+			FieldName = PointPropertySelector.GetName().ToString();
 		}
 		else
 		{
diff --git a/Engine/Plugins/Experimental/PCGExternalDataInterop/Source/PCGExternalDataInterop/Private/Alembic/PCGAlembicInterop.cpp b/Engine/Plugins/Experimental/PCGExternalDataInterop/Source/PCGExternalDataInterop/Private/Alembic/PCGAlembicInterop.cpp
--- a/Engine/Plugins/Experimental/PCGExternalDataInterop/Source/PCGExternalDataInterop/Private/Alembic/PCGAlembicInterop.cpp
+++ b/Engine/Plugins/Experimental/PCGExternalDataInterop/Source/PCGExternalDataInterop/Private/Alembic/PCGAlembicInterop.cpp
@@ -266,7 +266,7 @@ void ParseAlembicObject(FPCGExternalDataContext* Context, const Alembic::Abc::IO
 
 				FPCGAttributePropertyOutputSelector PointPositionSelector;
 				PointPositionSelector.ImportFromOtherSelector(*RemappedPositions);
-				const FString PropName = PointPositionSelector.GetDisplayText().ToString();
+				const FString PropName = PointPositionSelector.GetName().ToString();
 
 				if (CreatePointAccessorAndValidate(Context, PointData, AlembicPositionAccessor, PointPositionSelector, PropName, PointPropertyAccessor))
 				{
@@ -304,7 +304,7 @@ void ParseAlembicObject(FPCGExternalDataContext* Context, const Alembic::Abc::IO
 				if (const FPCGAttributePropertySelector* MappedField = Settings->AttributeMapping.Find(PropName))
 				{
 					PointPropertySelector = *MappedField;
-					PropName = PointPropertySelector.GetDisplayText().ToString();
+					PropName = PointPropertySelector.GetName().ToString();
 				}
 				else
 				{
diff --git a/Engine/Plugins/Experimental/PCGExternalDataInterop/Source/PCGExternalDataInterop/Private/Elements/PCGLoadAlembicElement.cpp b/Engine/Plugins/Experimental/PCGExternalDataInterop/Source/PCGExternalDataInterop/Private/Elements/PCGLoadAlembicElement.cpp
--- a/Engine/Plugins/Experimental/PCGExternalDataInterop/Source/PCGExternalDataInterop/Private/Elements/PCGLoadAlembicElement.cpp
+++ b/Engine/Plugins/Experimental/PCGExternalDataInterop/Source/PCGExternalDataInterop/Private/Elements/PCGLoadAlembicElement.cpp
@@ -8,6 +8,20 @@
 #define LOCTEXT_NAMESPACE "PCGLoadAlembic"
 
 #if WITH_EDITOR
+void UPCGLoadAlembicSettings::PostEditChangeProperty(FPropertyChangedEvent& PropertyChangedEvent)
+{
+	if (PropertyChangedEvent.Property && PropertyChangedEvent.Property->GetFName() == GET_MEMBER_NAME_CHECKED(UPCGLoadAlembicSettings, Setup))
+	{
+		if (Setup == EPCGLoadAlembicStandardSetup::CitySample)
+		{
+			SetupFromStandard(Setup);
+			Setup = EPCGLoadAlembicStandardSetup::None;
+		}
+	}
+
+	Super::PostEditChangeProperty(PropertyChangedEvent);
+}
+
 FText UPCGLoadAlembicSettings::GetDefaultNodeTitle() const
 {
 	return LOCTEXT("NodeTitle", "Load Alembic");
@@ -24,6 +38,41 @@ FPCGElementPtr UPCGLoadAlembicSettings::CreateElement() const
 	return MakeShared<FPCGLoadAlembicElement>();
 }
 
+void UPCGLoadAlembicSettings::SetupFromStandard(EPCGLoadAlembicStandardSetup InSetup)
+{
+	if (InSetup == EPCGLoadAlembicStandardSetup::CitySample)
+	{
+		SetupFromStandard(InSetup, ConversionScale, ConversionRotation, bConversionFlipHandedness, AttributeMapping);
+	}
+}
+
+void UPCGLoadAlembicSettings::SetupFromStandard(EPCGLoadAlembicStandardSetup InSetup, FVector& InConversionScale, FVector& InConversionRotation, bool& bInConversionFlipHandedness, TMap<FString, FPCGAttributePropertyInputSelector>& InAttributeMapping)
+{
+	if (InSetup == EPCGLoadAlembicStandardSetup::CitySample)
+	{
+		InConversionScale = FVector(1.0f, 1.0f, 1.0f);
+		InConversionRotation = FVector::ZeroVector;
+		bInConversionFlipHandedness = true;
+
+		InAttributeMapping.Reset();
+
+		FPCGAttributePropertyInputSelector PositionSelector;
+		PositionSelector.Update(FString(TEXT("$Position.xzy")));
+
+		InAttributeMapping.Add(FString(TEXT("position")), MoveTemp(PositionSelector));
+
+		FPCGAttributePropertyInputSelector ScaleSelector;
+		ScaleSelector.Update(FString(TEXT("$Scale.xzy")));
+
+		InAttributeMapping.Add(FString(TEXT("scale")), MoveTemp(ScaleSelector));
+
+		FPCGAttributePropertyInputSelector RotationSelector;
+		RotationSelector.Update(FString(TEXT("$Rotation.xzyw")));
+
+		InAttributeMapping.Add(FString(TEXT("orient")), MoveTemp(RotationSelector));
+	}
+}
+
 FPCGContext* FPCGLoadAlembicElement::Initialize(const FPCGDataCollection& InputData, TWeakObjectPtr<UPCGComponent> SourceComponent, const UPCGNode* Node)
 {
 	FPCGLoadAlembicContext* Context = new FPCGLoadAlembicContext();
@@ -80,9 +129,10 @@ bool FPCGLoadAlembicElement::ExecuteLoad(FPCGExternalDataContext* InContext) con
 
 	const FVector& ConversionScale = Settings->ConversionScale;
 	const FVector& ConversionRotation = Settings->ConversionRotation;
+	const bool bFlipRotationW = Settings->bConversionFlipHandedness;
 
 	const FTransform ConversionTransform(FRotator::MakeFromEuler(ConversionRotation), FVector::ZeroVector, ConversionScale);
-	if (!ConversionTransform.Equals(FTransform::Identity))
+	if (!ConversionTransform.Equals(FTransform::Identity) || bFlipRotationW)
 	{
 		for (const FPCGExternalDataContext::FPointDataAccessorsMapping& DataMapping : Context->PointDataAccessorsMapping)
 		{
@@ -98,6 +148,13 @@ bool FPCGLoadAlembicElement::ExecuteLoad(FPCGExternalDataContext* InContext) con
 			for (FPCGPoint& Point : Points)
 			{
 				Point.Transform = Point.Transform * ConversionTransform;
+
+				if (bFlipRotationW)
+				{
+					FQuat Rotation = Point.Transform.GetRotation();
+					Rotation.W *= -1;
+					Point.Transform.SetRotation(Rotation);
+				}
 			}
 		}
 	}
diff --git a/Engine/Plugins/Experimental/PCGExternalDataInterop/Source/PCGExternalDataInterop/Public/Elements/PCGLoadAlembicElement.h b/Engine/Plugins/Experimental/PCGExternalDataInterop/Source/PCGExternalDataInterop/Public/Elements/PCGLoadAlembicElement.h
--- a/Engine/Plugins/Experimental/PCGExternalDataInterop/Source/PCGExternalDataInterop/Public/Elements/PCGLoadAlembicElement.h
+++ b/Engine/Plugins/Experimental/PCGExternalDataInterop/Source/PCGExternalDataInterop/Public/Elements/PCGLoadAlembicElement.h
@@ -8,12 +8,25 @@
 
 struct FPCGExternalDataContext;
 
+UENUM()
+enum class EPCGLoadAlembicStandardSetup : uint8
+{
+	None = 0,
+	CitySample UMETA(Tooltip="Uses the same setup as in the City Sample demo: right handed Y-up and the orient and scale mapping to the rotation and scale, respectively")
+};
+
 UCLASS(BlueprintType, ClassGroup = (Procedural))
 class PCGEXTERNALDATAINTEROP_API UPCGLoadAlembicSettings : public UPCGExternalDataSettings
 {
 	GENERATED_BODY()
 
 public:
+	// ~Begin UObject interface
+#if WITH_EDITOR
+	virtual void PostEditChangeProperty(FPropertyChangedEvent& PropertyChangedEvent) override;
+#endif
+	// ~End UObject interface
+
 	// ~Begin UPCGSettings interface
 #if WITH_EDITOR
 	virtual FName GetDefaultNodeName() const override { return FName(TEXT("LoadAlembic")); }
@@ -37,6 +50,20 @@ public:
 	/** Rotation in Euler angles applied during import. For Max, use (90, 0, 0). */
 	UPROPERTY(BlueprintReadWrite, EditAnywhere, Category = "Alembic", meta = (PCG_Overridable))
 	FVector ConversionRotation = FVector::ZeroVector;
+
+	/** When changing handedness, it is sometimes needed to flip the rotation direction */
+	UPROPERTY(BlueprintReadWrite, EditAnywhere, Category = "Alembic", meta = (PCG_Overridable, Tooltip="Flips rotation direction (W), useful together with swizzling"))
+	bool bConversionFlipHandedness = false;
+
+	UFUNCTION(BlueprintCallable, Category = "PCG|AlembicImport")
+	void SetupFromStandard(EPCGLoadAlembicStandardSetup InSetup);
+
+	static void SetupFromStandard(EPCGLoadAlembicStandardSetup InSetup, FVector& InConversionScale, FVector& InConversionRotation, bool& bInConversionFlipHandedness, TMap<FString, FPCGAttributePropertyInputSelector>& InAttributeMapping);
+
+#if WITH_EDITORONLY_DATA
+	UPROPERTY(EditAnywhere, Category = "Alembic", meta = (DisplayName="Setup from standard"))
+	EPCGLoadAlembicStandardSetup Setup = EPCGLoadAlembicStandardSetup::None;
+#endif
 };
 
 class PCGEXTERNALDATAINTEROP_API FPCGLoadAlembicElement : public FPCGExternalDataElement
diff --git a/Engine/Plugins/Experimental/PCGExternalDataInterop/Source/PCGExternalDataInteropEditor/Private/PCGLoadAlembic.cpp b/Engine/Plugins/Experimental/PCGExternalDataInterop/Source/PCGExternalDataInteropEditor/Private/PCGLoadAlembic.cpp
--- a/Engine/Plugins/Experimental/PCGExternalDataInterop/Source/PCGExternalDataInteropEditor/Private/PCGLoadAlembic.cpp
+++ b/Engine/Plugins/Experimental/PCGExternalDataInterop/Source/PCGExternalDataInteropEditor/Private/PCGLoadAlembic.cpp
@@ -38,4 +38,13 @@ void UPCGLoadAlembicFunctionLibrary::LoadAlembicFileToPCG(const FPCGLoadAlembicB
 			}
 		}
 	}
+}
+
+void UPCGLoadAlembicFunctionLibrary::SetupFromStandard(FPCGLoadAlembicBPData& Data, EPCGLoadAlembicStandardSetup InSetup)
+{
+	if (InSetup != EPCGLoadAlembicStandardSetup::None)
+	{
+		Data.ConversionSettings.Preset = EAbcConversionPreset::Custom;
+		UPCGLoadAlembicSettings::SetupFromStandard(InSetup, Data.ConversionSettings.Scale, Data.ConversionSettings.Rotation, Data.bConversionFlipHandedness, Data.AttributeMapping);
+	}	
 }
\ No newline at end of file
diff --git a/Engine/Plugins/Experimental/PCGExternalDataInterop/Source/PCGExternalDataInteropEditor/Public/PCGLoadAlembic.h b/Engine/Plugins/Experimental/PCGExternalDataInterop/Source/PCGExternalDataInteropEditor/Public/PCGLoadAlembic.h
--- a/Engine/Plugins/Experimental/PCGExternalDataInterop/Source/PCGExternalDataInteropEditor/Public/PCGLoadAlembic.h
+++ b/Engine/Plugins/Experimental/PCGExternalDataInterop/Source/PCGExternalDataInteropEditor/Public/PCGLoadAlembic.h
@@ -2,6 +2,7 @@
 
 #pragma once
 
+#include "Elements/PCGLoadAlembicElement.h"
 #include "Metadata/PCGAttributePropertySelector.h"
 
 #include "AbcImportSettings.h"
@@ -20,6 +21,9 @@ struct PCGEXTERNALDATAINTEROPEDITOR_API FPCGLoadAlembicBPData
 	UPROPERTY(BlueprintReadWrite, EditAnywhere, Category = "Settings")
 	FAbcConversionSettings ConversionSettings;
 
+	UPROPERTY(BlueprintReadWrite, EditAnywhere, Category = "Settings", meta = (Tooltip = "Flips rotation direction (W), useful together with swizzling"))
+	bool bConversionFlipHandedness = false;
+
 	UPROPERTY(BlueprintReadWrite, EditAnywhere, Category = "Settings")
 	TMap<FString, FPCGAttributePropertyInputSelector> AttributeMapping;
 };
@@ -31,4 +35,7 @@ class PCGEXTERNALDATAINTEROPEDITOR_API UPCGLoadAlembicFunctionLibrary : public U
 
 	UFUNCTION(BlueprintCallable, Category = "PCG|IO")
 	static void LoadAlembicFileToPCG(const FPCGLoadAlembicBPData& Settings, FPCGDataCollection& Data, UObject* TargetOuter);
+
+	UFUNCTION(BlueprintCallable, Category = "PCG|IO", meta = (ScriptMethod))
+	static void SetupFromStandard(UPARAM(ref) FPCGLoadAlembicBPData& Data, EPCGLoadAlembicStandardSetup InSetup);
 };
\ No newline at end of file
