Fixed missing handedness check in the BP API used in the EUW for the alembic import

#rnx

diff --git a/Engine/Plugins/Experimental/PCGExternalDataInterop/Source/PCGExternalDataInteropEditor/Private/PCGLoadAlembic.cpp b/Engine/Plugins/Experimental/PCGExternalDataInterop/Source/PCGExternalDataInteropEditor/Private/PCGLoadAlembic.cpp
--- a/Engine/Plugins/Experimental/PCGExternalDataInterop/Source/PCGExternalDataInteropEditor/Private/PCGLoadAlembic.cpp
+++ b/Engine/Plugins/Experimental/PCGExternalDataInterop/Source/PCGExternalDataInteropEditor/Private/PCGLoadAlembic.cpp
@@ -11,6 +11,7 @@ void UPCGLoadAlembicFunctionLibrary::LoadAlembicFileToPCG(const FPCGLoadAlembicB
 	Settings->AlembicFilePath = InSettings.AlembicFilePath;
 	Settings->ConversionScale = InSettings.ConversionSettings.Scale;
 	Settings->ConversionRotation = InSettings.ConversionSettings.Rotation;
+	Settings->bConversionFlipHandedness = InSettings.bConversionFlipHandedness;
 	Settings->AttributeMapping = InSettings.AttributeMapping;
 
 	FPCGLoadAlembicElement Element;
