[PCG] Added some robustness around the alembic parser when there are no arbitrary geom parameters

#rnx

diff --git a/Engine/Plugins/PCGExternalDataInterop/Source/PCGExternalDataInterop/Private/Alembic/PCGAlembicInterop.cpp b/Engine/Plugins/PCGExternalDataInterop/Source/PCGExternalDataInterop/Private/Alembic/PCGAlembicInterop.cpp
--- a/Engine/Plugins/PCGExternalDataInterop/Source/PCGExternalDataInterop/Private/Alembic/PCGAlembicInterop.cpp
+++ b/Engine/Plugins/PCGExternalDataInterop/Source/PCGExternalDataInterop/Private/Alembic/PCGAlembicInterop.cpp
@@ -284,37 +284,40 @@ void ParseAlembicObject(FPCGExternalDataContext* Context, const Alembic::Abc::IO
 			}
 
 			Alembic::AbcGeom::ICompoundProperty Parameters = Points.getSchema().getArbGeomParams();
-			for (int Index = 0; Index < Parameters.getNumProperties(); ++Index)
+			if (Parameters.valid())
 			{
-				Alembic::Abc::PropertyHeader PropertyHeader = Parameters.getPropertyHeader(Index);
-				FString PropName(PropertyHeader.getName().c_str());
-
-				// We'll parse only properties that affect every point/object
-				if (PropertyHeader.getPropertyType() != Alembic::Abc::kArrayProperty &&
-					!(PropertyHeader.getPropertyType() == Alembic::Abc::kCompoundProperty && (PropertyHeader.getDataType().getPod() == Alembic::Util::kUnknownPOD || PropertyHeader.getDataType().getPod() == Alembic::Util::kStringPOD)))
-				{
-					PCGE_LOG_C(Error, GraphAndLog, Context, FText::Format(LOCTEXT("AlembicPropertyType", "Property '{0}' is not supported (expected array)."), FText::FromString(PropName)));
-					continue;
-				}
-
-				TUniquePtr<const IPCGAttributeAccessor> AlembicPropAccessor = CreateAlembicPropAccessor(Parameters, PropertyHeader, PropName);
-				TUniquePtr<IPCGAttributeAccessor> PointPropertyAccessor;
-
-				// Setup attribute property selector
-				FPCGAttributePropertySelector PointPropertySelector;
-				if (const FPCGAttributePropertySelector* MappedField = Settings->AttributeMapping.Find(PropName))
-				{
-					PointPropertySelector = *MappedField;
-					PropName = PointPropertySelector.GetName().ToString();
-				}
-				else
-				{
-					PointPropertySelector.Update(PropName);
-				}
-
-				if (CreatePointAccessorAndValidate(Context, PointData, AlembicPropAccessor, PointPropertySelector, PropName, PointPropertyAccessor))
+				for (int Index = 0; Index < Parameters.getNumProperties(); ++Index)
 				{
-					PointDataAccessorMapping.RowToPointAccessors.Emplace(MoveTemp(AlembicPropAccessor), MoveTemp(PointPropertyAccessor), PointPropertySelector);
+					Alembic::Abc::PropertyHeader PropertyHeader = Parameters.getPropertyHeader(Index);
+					FString PropName(PropertyHeader.getName().c_str());
+
+					// We'll parse only properties that affect every point/object
+					if (PropertyHeader.getPropertyType() != Alembic::Abc::kArrayProperty &&
+						!(PropertyHeader.getPropertyType() == Alembic::Abc::kCompoundProperty && (PropertyHeader.getDataType().getPod() == Alembic::Util::kUnknownPOD || PropertyHeader.getDataType().getPod() == Alembic::Util::kStringPOD)))
+					{
+						PCGE_LOG_C(Error, GraphAndLog, Context, FText::Format(LOCTEXT("AlembicPropertyType", "Property '{0}' is not supported (expected array)."), FText::FromString(PropName)));
+						continue;
+					}
+
+					TUniquePtr<const IPCGAttributeAccessor> AlembicPropAccessor = CreateAlembicPropAccessor(Parameters, PropertyHeader, PropName);
+					TUniquePtr<IPCGAttributeAccessor> PointPropertyAccessor;
+
+					// Setup attribute property selector
+					FPCGAttributePropertySelector PointPropertySelector;
+					if (const FPCGAttributePropertySelector* MappedField = Settings->AttributeMapping.Find(PropName))
+					{
+						PointPropertySelector = *MappedField;
+						PropName = PointPropertySelector.GetName().ToString();
+					}
+					else
+					{
+						PointPropertySelector.Update(PropName);
+					}
+
+					if (CreatePointAccessorAndValidate(Context, PointData, AlembicPropAccessor, PointPropertySelector, PropName, PointPropertyAccessor))
+					{
+						PointDataAccessorMapping.RowToPointAccessors.Emplace(MoveTemp(AlembicPropAccessor), MoveTemp(PointPropertyAccessor), PointPropertySelector);
+					}
 				}
 			}
 		}
