[PCG] Add warning when trying to import an alembic file with no point cloud data or no points.

#rnx

diff --git a/Engine/Plugins/PCGExternalDataInterop/Source/PCGExternalDataInterop/Private/Alembic/PCGAlembicInterop.cpp b/Engine/Plugins/PCGExternalDataInterop/Source/PCGExternalDataInterop/Private/Alembic/PCGAlembicInterop.cpp
--- a/Engine/Plugins/PCGExternalDataInterop/Source/PCGExternalDataInterop/Private/Alembic/PCGAlembicInterop.cpp
+++ b/Engine/Plugins/PCGExternalDataInterop/Source/PCGExternalDataInterop/Private/Alembic/PCGAlembicInterop.cpp
@@ -225,7 +225,7 @@ bool CreatePointAccessorAndValidate(FPCGContext* Context, UPCGPointData* PointDa
 	return true;
 }
 
-void ParseAlembicObject(FPCGExternalDataContext* Context, const Alembic::Abc::IObject& Object)
+bool ParseAlembicObject(FPCGExternalDataContext* Context, const Alembic::Abc::IObject& Object)
 {
 	check(Context);
 	const UPCGExternalDataSettings* Settings = Context->GetInputSettings<UPCGExternalDataSettings>();
@@ -234,6 +234,8 @@ void ParseAlembicObject(FPCGExternalDataContext* Context, const Alembic::Abc::IO
 	const Alembic::Abc::MetaData& ObjectMetaData = Object.getMetaData();
 	const uint32 NumChildren = Object.getNumChildren();
 
+	bool bHasCreatedData = false;
+
 	if (Alembic::AbcGeom::IPoints::matches(ObjectMetaData))
 	{
 		Alembic::AbcGeom::IPoints Points = Alembic::AbcGeom::IPoints(Object, Alembic::Abc::kWrapExisting);
@@ -244,6 +246,8 @@ void ParseAlembicObject(FPCGExternalDataContext* Context, const Alembic::Abc::IO
 
 		if (NumPoints > 0)
 		{
+			bHasCreatedData = true;
+
 			// Create point data & mapping
 			UPCGPointData* PointData = NewObject<UPCGPointData>();
 			check(PointData);
@@ -327,9 +331,11 @@ void ParseAlembicObject(FPCGExternalDataContext* Context, const Alembic::Abc::IO
 	{
 		for (uint32 ChildIndex = 0; ChildIndex < NumChildren; ++ChildIndex)
 		{
-			ParseAlembicObject(Context, Object.getChild(ChildIndex));
+			bHasCreatedData |= ParseAlembicObject(Context, Object.getChild(ChildIndex));
 		}
 	}
+
+	return bHasCreatedData;
 }
 
 void LoadFromAlembicFile(FPCGLoadAlembicContext* Context, const FString& FileName)
@@ -365,7 +371,11 @@ void LoadFromAlembicFile(FPCGLoadAlembicContext* Context, const FString& FileNam
 		return;
 	}
 
-	ParseAlembicObject(Context, TopObject);
+	if (!ParseAlembicObject(Context, TopObject))
+	{
+		PCGE_LOG_C(Warning, GraphAndLog, Context, FText::Format(LOCTEXT("FailedToCreatePointData", "Import of '{0}' was successful but there is no valid point cloud data in the Alembic file."), FText::FromString(FileName)));
+		return;
+	}
 }
 
 #endif // WITH_EDITOR
