[PCG] Improved caching on load alembic file, so it reruns properly if the file has changed.

#rnx

diff --git a/Engine/Plugins/PCGInterops/PCGExternalDataInterop/Source/PCGExternalDataInterop/Private/Elements/PCGLoadAlembicElement.cpp b/Engine/Plugins/PCGInterops/PCGExternalDataInterop/Source/PCGExternalDataInterop/Private/Elements/PCGLoadAlembicElement.cpp
--- a/Engine/Plugins/PCGInterops/PCGExternalDataInterop/Source/PCGExternalDataInterop/Private/Elements/PCGLoadAlembicElement.cpp
+++ b/Engine/Plugins/PCGInterops/PCGExternalDataInterop/Source/PCGExternalDataInterop/Private/Elements/PCGLoadAlembicElement.cpp
@@ -3,6 +3,7 @@
 #include "Elements/PCGLoadAlembicElement.h"
 
 #include "Alembic/PCGAlembicInterop.h"
+#include "HAL/FileManager.h"
 #include "Misc/Paths.h"
 
 #define LOCTEXT_NAMESPACE "PCGLoadAlembic"
@@ -73,6 +74,23 @@ void UPCGLoadAlembicSettings::SetupFromStandard(EPCGLoadAlembicStandardSetup InS
 	}
 }
 
+void FPCGLoadAlembicElement::GetDependenciesCrc(const FPCGDataCollection& InInput, const UPCGSettings* InSettings, UPCGComponent* InComponent, FPCGCrc& OutCrc) const
+{
+	FPCGCrc Crc;
+	FPCGExternalDataElement::GetDependenciesCrc(InInput, InSettings, InComponent, Crc);
+
+	if (const UPCGLoadAlembicSettings* Settings = Cast<UPCGLoadAlembicSettings>(InSettings))
+	{
+		const FDateTime FileTimeStamp = IFileManager::Get().GetTimeStamp(*Settings->AlembicFilePath.FilePath);
+		if (FileTimeStamp != FDateTime::MinValue())
+		{
+			Crc.Combine(GetTypeHash(FileTimeStamp));
+		}
+	}
+
+	OutCrc = Crc;
+}
+
 FPCGContext* FPCGLoadAlembicElement::CreateContext()
 {
 	return new FPCGLoadAlembicContext();
diff --git a/Engine/Plugins/PCGInterops/PCGExternalDataInterop/Source/PCGExternalDataInterop/Public/Elements/PCGLoadAlembicElement.h b/Engine/Plugins/PCGInterops/PCGExternalDataInterop/Source/PCGExternalDataInterop/Public/Elements/PCGLoadAlembicElement.h
--- a/Engine/Plugins/PCGInterops/PCGExternalDataInterop/Source/PCGExternalDataInterop/Public/Elements/PCGLoadAlembicElement.h
+++ b/Engine/Plugins/PCGInterops/PCGExternalDataInterop/Source/PCGExternalDataInterop/Public/Elements/PCGLoadAlembicElement.h
@@ -68,6 +68,8 @@ public:
 
 class PCGEXTERNALDATAINTEROP_API FPCGLoadAlembicElement : public FPCGExternalDataElement
 {
+public:
+	virtual void GetDependenciesCrc(const FPCGDataCollection& InInput, const UPCGSettings* InSettings, UPCGComponent* InComponent, FPCGCrc& OutCrc) const override;
 protected:
 	virtual FPCGContext* CreateContext() override;
 	virtual bool PrepareLoad(FPCGExternalDataContext* Context) const override;
