Fix compile error because winnt.h is included by Alembic includes, which redefines MemoryBarrier, that then causes compile issues in anything that subsequently includes WindowsPlatformMisc.
This has only worked so far because of PCH
PCGAlembicInterop.cpp does not compile on it's own without this fix.

diff --git a/Engine/Plugins/PCGInterops/PCGExternalDataInterop/Source/PCGExternalDataInterop/Private/Alembic/PCGAlembicInterop.h b/Engine/Plugins/PCGInterops/PCGExternalDataInterop/Source/PCGExternalDataInterop/Private/Alembic/PCGAlembicInterop.h
--- a/Engine/Plugins/PCGInterops/PCGExternalDataInterop/Source/PCGExternalDataInterop/Private/Alembic/PCGAlembicInterop.h
+++ b/Engine/Plugins/PCGInterops/PCGExternalDataInterop/Source/PCGExternalDataInterop/Private/Alembic/PCGAlembicInterop.h
@@ -14,6 +14,13 @@ THIRD_PARTY_INCLUDES_START
 #include <Alembic/AbcCoreOgawa/All.h>
 #include "Alembic/AbcGeom/All.h"
 #include "Alembic/AbcCoreFactory/IFactory.h"
+// The above includes bring in winnt.h on windows builds, which has a #define for MemoryBarrier, which then
+// conflicts with anything that brings in WindowsPlatformMisc.h afterwards, such as Queue.h
+// This happens with PCGAlembicInterop.cpp and that file does not compile without PCH because of this. We have just been lucky
+// in our include order
+#ifdef MemoryBarrier
+#undef MemoryBarrier
+#endif // MemoryBarrier
 #pragma warning(pop)
 THIRD_PARTY_INCLUDES_END
 
