[PCG] Better management of windows headers wrt alembic files

#rnx

diff --git a/Engine/Plugins/PCGInterops/PCGExternalDataInterop/Source/PCGExternalDataInterop/Private/Elements/PCGLoadAlembicElement.cpp b/Engine/Plugins/PCGInterops/PCGExternalDataInterop/Source/PCGExternalDataInterop/Private/Elements/PCGLoadAlembicElement.cpp
--- a/Engine/Plugins/PCGInterops/PCGExternalDataInterop/Source/PCGExternalDataInterop/Private/Elements/PCGLoadAlembicElement.cpp
+++ b/Engine/Plugins/PCGInterops/PCGExternalDataInterop/Source/PCGExternalDataInterop/Private/Elements/PCGLoadAlembicElement.cpp
@@ -20,28 +20,23 @@
 #include "AbcImportSettings.h"
 #include "AbcUtilities.h"
 
+#if PLATFORM_WINDOWS
+#include "Windows/AllowWindowsPlatformTypes.h"
+#endif
+
 THIRD_PARTY_INCLUDES_START
 #pragma warning(push)
 #pragma warning(disable:4005) // TEXT macro redefinition
 #include <Alembic/AbcCoreOgawa/All.h>
 #include "Alembic/AbcGeom/All.h"
 #include "Alembic/AbcCoreFactory/IFactory.h"
-// The above includes bring in winnt.h on windows builds, which has a #define for MemoryBarrier, which then
-// conflicts with anything that brings in WindowsPlatformMisc.h afterwards, such as Queue.h
-// This happens with PCGAlembicInterop.cpp and that file does not compile without PCH because of this. We have just been lucky
-// in our include order
-#ifdef MemoryBarrier
-#undef MemoryBarrier
-#endif // MemoryBarrier
 #pragma warning(pop)
+THIRD_PARTY_INCLUDES_END
 
-// In a similar vein, this can include window.s which will define GetObject as GetObjectW.
-#ifdef GetObject
-#undef GetObject
+#if PLATFORM_WINDOWS
+#include "Windows/HideWindowsPlatformTypes.h"
 #endif
 
-THIRD_PARTY_INCLUDES_END
-
 #endif // WITH_EDITOR
 
 struct FPCGLoadAlembicContext : public FPCGExternalDataContext
